Include capitalization correction to make mac CIS happy.

diff --git a/Engine/Plugins/Runtime/GoogleARCore/Source/GoogleARCoreBase/Private/GoogleARCoreXRCamera.cpp b/Engine/Plugins/Runtime/GoogleARCore/Source/GoogleARCoreBase/Private/GoogleARCoreXRCamera.cpp
--- a/Engine/Plugins/Runtime/GoogleARCore/Source/GoogleARCoreBase/Private/GoogleARCoreXRCamera.cpp
+++ b/Engine/Plugins/Runtime/GoogleARCore/Source/GoogleARCoreBase/Private/GoogleARCoreXRCamera.cpp
@@ -1,6 +1,6 @@
 // Copyright 1998-2017 Epic Games, Inc. All Rights Reserved.
 
-#include "GoogleARCoreXRcamera.h"
+#include "GoogleARCoreXRCamera.h"
 #include "GoogleARCoreHMD.h"
 #include "SceneView.h"
 
