AppleARKit and GoogleARCore plugins are no longer enabled by default.

diff --git a/Engine/Plugins/Runtime/AppleARKit/AppleARKit.uplugin b/Engine/Plugins/Runtime/AppleARKit/AppleARKit.uplugin
--- a/Engine/Plugins/Runtime/AppleARKit/AppleARKit.uplugin
+++ b/Engine/Plugins/Runtime/AppleARKit/AppleARKit.uplugin
@@ -12,7 +12,7 @@
 	"SupportURL": "",
 	"CanContainContent": true,
 	"IsBetaVersion": true,
-    "EnabledByDefault": true,
+	"EnabledByDefault": false,
 	"Modules": [
 		{
 			"Name": "AppleARKit",
diff --git a/Engine/Plugins/Runtime/GoogleARCore/GoogleARCore.uplugin b/Engine/Plugins/Runtime/GoogleARCore/GoogleARCore.uplugin
--- a/Engine/Plugins/Runtime/GoogleARCore/GoogleARCore.uplugin
+++ b/Engine/Plugins/Runtime/GoogleARCore/GoogleARCore.uplugin
@@ -8,7 +8,9 @@
     "Category" : "Augmented Reality",
     "CreatedBy" : "Google",
     "CreatedByURL" : "https://developers.google.com/tango/",
-    "CanContainContent" : true,
+    "CanContainContent": true,
+    "IsBetaVersion": true,
+    "EnabledByDefault": false,
     "Modules" :
     [
         {
