Fix include path on Android.

diff --git a/Engine/Plugins/Runtime/GoogleARCore/Source/GoogleARCoreBase/Private/GoogleARCoreCameraImageBlitter.cpp b/Engine/Plugins/Runtime/GoogleARCore/Source/GoogleARCoreBase/Private/GoogleARCoreCameraImageBlitter.cpp
--- a/Engine/Plugins/Runtime/GoogleARCore/Source/GoogleARCoreBase/Private/GoogleARCoreCameraImageBlitter.cpp
+++ b/Engine/Plugins/Runtime/GoogleARCore/Source/GoogleARCoreBase/Private/GoogleARCoreCameraImageBlitter.cpp
@@ -7,7 +7,7 @@
 #include "GoogleARCorePassthroughCameraRenderer.h"
 
 #if PLATFORM_ANDROID
-#include "AndroidApplication.h"
+#include "Android/AndroidApplication.h"
 #include <GLES2/gl2.h>
 #include <GLES2/gl2ext.h>
 #endif
