Fix for compilation error relying on editor-only functionality

diff --git a/Engine/Plugins/Runtime/GoogleARCore/Source/GoogleARCoreBase/Private/GoogleARCoreAugmentedImageDatabase.cpp b/Engine/Plugins/Runtime/GoogleARCore/Source/GoogleARCoreBase/Private/GoogleARCoreAugmentedImageDatabase.cpp
--- a/Engine/Plugins/Runtime/GoogleARCore/Source/GoogleARCoreBase/Private/GoogleARCoreAugmentedImageDatabase.cpp
+++ b/Engine/Plugins/Runtime/GoogleARCore/Source/GoogleARCoreBase/Private/GoogleARCoreAugmentedImageDatabase.cpp
@@ -139,7 +139,7 @@ void UGoogleARCoreAugmentedImageDatabase::Serialize(FArchive& Ar)
 	FString OutStderr;
 	int32 OutReturnCode = 0;
 
-#if !PLATFORM_ANDROID
+#if !PLATFORM_ANDROID && WITH_EDITORONLY_DATA
 
 	if (!Ar.IsLoading() && Ar.IsCooking()) {
 
