Remove call to ArSession_checkSupported which always reported true and is now deprecated

diff --git a/Engine/Plugins/Runtime/GoogleARCore/Source/GoogleARCoreBase/Private/GoogleARCoreAPI.cpp b/Engine/Plugins/Runtime/GoogleARCore/Source/GoogleARCoreBase/Private/GoogleARCoreAPI.cpp
--- a/Engine/Plugins/Runtime/GoogleARCore/Source/GoogleARCoreBase/Private/GoogleARCoreAPI.cpp
+++ b/Engine/Plugins/Runtime/GoogleARCore/Source/GoogleARCoreBase/Private/GoogleARCoreAPI.cpp
@@ -226,10 +226,9 @@ bool FGoogleARCoreSession::IsConfigSupported(const UARSessionConfig& Config)
 	ArConfig_setPlaneFindingMode(SessionHandle, NewConfigHandle, static_cast<ArPlaneFindingMode>(Config.GetPlaneDetectionMode()));
 	ArConfig_setUpdateMode(SessionHandle, NewConfigHandle, static_cast<ArUpdateMode>(Config.GetFrameSyncMode()));
 
-	ArStatus Status = ArSession_checkSupported(SessionHandle, NewConfigHandle);
 	ArConfig_destroy(NewConfigHandle);
 
-	return Status == ArStatus::AR_SUCCESS;
+	return true;
 #endif
 	return false;
 }
