Fix for static analysis warnings in ARCore code

diff --git a/Engine/Plugins/Runtime/GoogleARCore/Source/GoogleARCoreBase/Private/GoogleARCoreAugmentedImageDatabase.cpp b/Engine/Plugins/Runtime/GoogleARCore/Source/GoogleARCoreBase/Private/GoogleARCoreAugmentedImageDatabase.cpp
--- a/Engine/Plugins/Runtime/GoogleARCore/Source/GoogleARCoreBase/Private/GoogleARCoreAugmentedImageDatabase.cpp
+++ b/Engine/Plugins/Runtime/GoogleARCore/Source/GoogleARCoreBase/Private/GoogleARCoreAugmentedImageDatabase.cpp
@@ -46,69 +46,76 @@ namespace
 			int32 Height = Tex->Source.GetSizeY();
 
 			png_structp PngPtr = png_create_write_struct(PNG_LIBPNG_VER_STRING, nullptr, nullptr, nullptr);
-			png_infop PngInfoPtr = png_create_info_struct(PngPtr);
-
-			// We're using C file IO here just because of libPNG interop.
-			FILE *OutputFile = fopen(TCHAR_TO_ANSI(*Filename), "wb+");
-
-			if (setjmp(png_jmpbuf(PngPtr)))
-			{
-				UE_LOG(
-					LogGoogleARCoreAPI, Error,
-					TEXT("Error writing PNG for texture %s."),
-					*Tex->GetName());
-				Ret = false;
-			}
-			else
+			if (PngPtr != nullptr)
 			{
-				png_init_io(PngPtr, OutputFile);
-				png_set_IHDR(
-					PngPtr, PngInfoPtr, Width, Height,
-					8, PNG_COLOR_TYPE_RGB, PNG_INTERLACE_NONE,
-					PNG_COMPRESSION_TYPE_BASE, PNG_FILTER_TYPE_BASE);
-				png_write_info(PngPtr, PngInfoPtr);
+				png_infop PngInfoPtr = png_create_info_struct(PngPtr);
 
-				uint8_t *Row = new uint8_t[Width * 3];
-
-				for (int32 y = 0; y < Height; y++)
+				if (PngInfoPtr != nullptr)
 				{
-					for (int32 x = 0; x < Width; x++)
+					// We're using C file IO here just because of libPNG interop.
+					FILE *OutputFile = fopen(TCHAR_TO_ANSI(*Filename), "wb+");
+
+					if (setjmp(png_jmpbuf(PngPtr)))
+					{
+						UE_LOG(
+							LogGoogleARCoreAPI, Error,
+							TEXT("Error writing PNG for texture %s."),
+							*Tex->GetName());
+						Ret = false;
+					}
+					else
 					{
-						for (int32 c = 0; c < 3; c++)
+						png_init_io(PngPtr, OutputFile);
+						png_set_IHDR(
+							PngPtr, PngInfoPtr, Width, Height,
+							8, PNG_COLOR_TYPE_RGB, PNG_INTERLACE_NONE,
+							PNG_COMPRESSION_TYPE_BASE, PNG_FILTER_TYPE_BASE);
+						png_write_info(PngPtr, PngInfoPtr);
+
+						uint8_t *Row = new uint8_t[Width * 3];
+
+						for (int32 y = 0; y < Height; y++)
 						{
-							int32 RealChannel = c;
-							if (Tex->Source.GetFormat() == TSF_BGRA8)
+							for (int32 x = 0; x < Width; x++)
 							{
-								if (RealChannel == 0) {
-									RealChannel = 2;
-								} else if (RealChannel == 2) {
-									RealChannel = 0;
+								for (int32 c = 0; c < 3; c++)
+								{
+									int32 RealChannel = c;
+									if (Tex->Source.GetFormat() == TSF_BGRA8)
+									{
+										if (RealChannel == 0) {
+											RealChannel = 2;
+										}
+										else if (RealChannel == 2) {
+											RealChannel = 0;
+										}
+									}
+									Row[x * 3 + c] = MipData[(y * Width + x) * 4 + RealChannel];
 								}
 							}
-							Row[x * 3 + c] = MipData[(y * Width + x) * 4 + RealChannel];
+							png_write_row(PngPtr, Row);
 						}
-					}
-					png_write_row(PngPtr, Row);
-				}
 
-				png_write_end(PngPtr, NULL);
+						png_write_end(PngPtr, NULL);
 
-				delete[] Row;
-			}
+						delete[] Row;
+					}
 
-			if (OutputFile)
-			{
-				fclose(OutputFile);
-			}
+					if (OutputFile)
+					{
+						fclose(OutputFile);
+					}
 
-			if (PngInfoPtr)
-			{
-				png_free_data(PngPtr, PngInfoPtr, PNG_FREE_ALL, -1);
-			}
+					if (PngInfoPtr)
+					{
+						png_free_data(PngPtr, PngInfoPtr, PNG_FREE_ALL, -1);
+					}
 
-			if (PngPtr)
-			{
-				png_destroy_write_struct(&PngPtr, nullptr);
+					if (PngPtr)
+					{
+						png_destroy_write_struct(&PngPtr, nullptr);
+					}
+				}
 			}
 		}
 		else
