Fix ctor init build warning.

diff --git a/Engine/Plugins/Runtime/GoogleARCore/Source/GoogleARCoreBase/Public/GoogleARCoreAugmentedImageDatabase.h b/Engine/Plugins/Runtime/GoogleARCore/Source/GoogleARCoreBase/Public/GoogleARCoreAugmentedImageDatabase.h
--- a/Engine/Plugins/Runtime/GoogleARCore/Source/GoogleARCoreBase/Public/GoogleARCoreAugmentedImageDatabase.h
+++ b/Engine/Plugins/Runtime/GoogleARCore/Source/GoogleARCoreBase/Public/GoogleARCoreAugmentedImageDatabase.h
@@ -34,7 +34,9 @@ struct FGoogleARCoreAugmentedImageDatabaseEntry
 	float Width;
 
 	FGoogleARCoreAugmentedImageDatabaseEntry()
-		: Width(0.0f) { }
+		: ImageAsset(nullptr)
+	    , Width(0.0f) 
+	{ }
 };
 
 UCLASS(BlueprintType)
