UE-71995 AR camera pass through is flipped when viewing in certain orientations
-Going back to the old shader code for passthrough camera.

diff --git a/Engine/Plugins/Experimental/RemoteSession/Source/RemoteSession/Private/Channels/RemoteSessionARCameraChannel.cpp b/Engine/Plugins/Experimental/RemoteSession/Source/RemoteSession/Private/Channels/RemoteSessionARCameraChannel.cpp
--- a/Engine/Plugins/Experimental/RemoteSession/Source/RemoteSession/Private/Channels/RemoteSessionARCameraChannel.cpp
+++ b/Engine/Plugins/Experimental/RemoteSession/Source/RemoteSession/Private/Channels/RemoteSessionARCameraChannel.cpp
@@ -78,6 +78,7 @@ public:
 		FMaterialShader::ModifyCompilationEnvironment(Platform, OutEnvironment);
 		OutEnvironment.SetDefine(TEXT("POST_PROCESS_MATERIAL"), 1);
 		OutEnvironment.SetDefine(TEXT("POST_PROCESS_MATERIAL_BEFORE_TONEMAP"), (Material->GetBlendableLocation() != BL_AfterTonemapping) ? 1 : 0);
+		OutEnvironment.SetDefine(TEXT("POST_PROCESS_AR_PASSTHROUGH"), 1);
 	}
 
 	FRemoteSessionARCameraVS() { }
diff --git a/Engine/Plugins/Runtime/AR/Apple/AppleARKit/Source/AppleARKit/Private/AppleARKitVideoOverlay.cpp b/Engine/Plugins/Runtime/AR/Apple/AppleARKit/Source/AppleARKit/Private/AppleARKitVideoOverlay.cpp
--- a/Engine/Plugins/Runtime/AR/Apple/AppleARKit/Source/AppleARKit/Private/AppleARKitVideoOverlay.cpp
+++ b/Engine/Plugins/Runtime/AR/Apple/AppleARKit/Source/AppleARKit/Private/AppleARKitVideoOverlay.cpp
@@ -276,6 +276,7 @@ public:
 		FMaterialShader::ModifyCompilationEnvironment(Platform, OutEnvironment);
 		OutEnvironment.SetDefine(TEXT("POST_PROCESS_MATERIAL"), 1);
 		OutEnvironment.SetDefine(TEXT("POST_PROCESS_MATERIAL_BEFORE_TONEMAP"), (Material->GetBlendableLocation() != BL_AfterTonemapping) ? 1 : 0);
+		OutEnvironment.SetDefine(TEXT("POST_PROCESS_AR_PASSTHROUGH"), 1);
 	}
 
 	FARKitCameraOverlayVS() { }
diff --git a/Engine/Plugins/Runtime/AR/Google/GoogleARCore/Source/GoogleARCoreRendering/Private/GoogleARCorePassthroughCameraRenderer.cpp b/Engine/Plugins/Runtime/AR/Google/GoogleARCore/Source/GoogleARCoreRendering/Private/GoogleARCorePassthroughCameraRenderer.cpp
--- a/Engine/Plugins/Runtime/AR/Google/GoogleARCore/Source/GoogleARCoreRendering/Private/GoogleARCorePassthroughCameraRenderer.cpp
+++ b/Engine/Plugins/Runtime/AR/Google/GoogleARCore/Source/GoogleARCoreRendering/Private/GoogleARCorePassthroughCameraRenderer.cpp
@@ -171,8 +171,8 @@ public:
 		FMaterialShader::ModifyCompilationEnvironment(Platform, OutEnvironment);
 
 		OutEnvironment.SetDefine(TEXT("POST_PROCESS_MATERIAL"), 1);
-
 		OutEnvironment.SetDefine(TEXT("POST_PROCESS_MATERIAL_BEFORE_TONEMAP"), (Material->GetBlendableLocation() != BL_AfterTonemapping) ? 1 : 0);
+		OutEnvironment.SetDefine(TEXT("POST_PROCESS_AR_PASSTHROUGH"), 1);
 	}
 
 	FGoogleARCoreCameraOverlayVS( )	{ }
diff --git a/Engine/Shaders/Private/PostProcessMaterialShaders.usf b/Engine/Shaders/Private/PostProcessMaterialShaders.usf
--- a/Engine/Shaders/Private/PostProcessMaterialShaders.usf
+++ b/Engine/Shaders/Private/PostProcessMaterialShaders.usf
@@ -215,7 +215,11 @@ void MainVS_ES2(
 	UNROLL
 	for (int CoordinateIndex = 0; CoordinateIndex < NUM_MATERIAL_TEXCOORDS; CoordinateIndex++)
 	{
+#if POST_PROCESS_AR_PASSTHROUGH
+		SetUV(Output, CoordinateIndex, OutUV);
+#else
 		SetUV(Output, CoordinateIndex, InPosition.xy);
+#endif
 	}
 	UNROLL
 	for (int CoordinateIndex = NUM_MATERIAL_TEXCOORDS; CoordinateIndex < NUM_TEX_COORD_INTERPOLATORS; CoordinateIndex++)
diff --git a/Engine/Shaders/Public/ShaderVersion.ush b/Engine/Shaders/Public/ShaderVersion.ush
--- a/Engine/Shaders/Public/ShaderVersion.ush
+++ b/Engine/Shaders/Public/ShaderVersion.ush
@@ -3,4 +3,4 @@
 // in Platform.ush (which should be included in any shader) it allows to invalidate the shader DDC.
 // 
 // If you are merging streams and there is a conflict with this GUID you should make a new GUID rather than taking one or the other.
-// GUID = 228AA72549B318E4175C78AB57A4F8AC
\ No newline at end of file
+// GUID = C15DFE5B413C0C4E6C49949A06F0AE74
\ No newline at end of file
