Fix for AR Camera Pass Through Materials (on ARCore)

diff --git a/Engine/Plugins/Runtime/AR/Google/GoogleARCore/Source/GoogleARCoreRendering/Private/GoogleARCorePassthroughCameraRenderer.cpp b/Engine/Plugins/Runtime/AR/Google/GoogleARCore/Source/GoogleARCoreRendering/Private/GoogleARCorePassthroughCameraRenderer.cpp
--- a/Engine/Plugins/Runtime/AR/Google/GoogleARCore/Source/GoogleARCoreRendering/Private/GoogleARCorePassthroughCameraRenderer.cpp
+++ b/Engine/Plugins/Runtime/AR/Google/GoogleARCore/Source/GoogleARCoreRendering/Private/GoogleARCorePassthroughCameraRenderer.cpp
@@ -93,6 +93,9 @@ void FGoogleARCorePassthroughCameraRenderer::InitializeRenderer_RenderThread(FTe
 
 	FExternalTextureRegistry::Get().RegisterExternalTexture(GoogleARCorePassthroughCameraExternalTextureGuid, VideoTexture, SamplerStateRHI);
 
+	//Make sure AR camera pass through materials are updated properly
+	FMaterialRenderProxy::UpdateDeferredCachedUniformExpressions();
+
 	bInitialized = true;
 }
 
