UE-65107 Orientations are inverted on GoogleAR
-Fixed logic that decides if we should flip or not, it was backward.

diff --git a/Engine/Plugins/Runtime/GoogleARCore/Source/GoogleARCoreRendering/Private/GoogleARCorePassthroughCameraRenderer.cpp b/Engine/Plugins/Runtime/GoogleARCore/Source/GoogleARCoreRendering/Private/GoogleARCorePassthroughCameraRenderer.cpp
--- a/Engine/Plugins/Runtime/GoogleARCore/Source/GoogleARCoreRendering/Private/GoogleARCorePassthroughCameraRenderer.cpp
+++ b/Engine/Plugins/Runtime/GoogleARCore/Source/GoogleARCoreRendering/Private/GoogleARCorePassthroughCameraRenderer.cpp
@@ -100,7 +100,7 @@ void FGoogleARCorePassthroughCameraRenderer::UpdateOverlayUVCoordinate_RenderThr
 {
 	check(InOverlayUVs.Num() == 8);
 
-	bool bNeedToFlipCameraImage = !RHINeedsToSwitchVerticalAxis(GShaderPlatformForFeatureLevel[GMaxRHIFeatureLevel]) || IsMobileHDR();
+	bool bNeedToFlipCameraImage = RHINeedsToSwitchVerticalAxis(GShaderPlatformForFeatureLevel[GMaxRHIFeatureLevel]) && !IsMobileHDR();
 
 	if (bNeedToFlipCameraImage)
 	{
