UE-65107 Orientations are inverted on GoogleAR
-This is probably a hack, but landscape and portait now orient the passthrough camera image correctly both in hdr and nohdr.
-Perhaps a better fix is possible closer to the image ingestion?
-Also deleted unused passthrough image copy function.

diff --git a/Engine/Plugins/Runtime/GoogleARCore/Source/GoogleARCoreBase/Private/GoogleARCoreAndroidHelper.cpp b/Engine/Plugins/Runtime/GoogleARCore/Source/GoogleARCoreBase/Private/GoogleARCoreAndroidHelper.cpp
--- a/Engine/Plugins/Runtime/GoogleARCore/Source/GoogleARCoreBase/Private/GoogleARCoreAndroidHelper.cpp
+++ b/Engine/Plugins/Runtime/GoogleARCore/Source/GoogleARCoreBase/Private/GoogleARCoreAndroidHelper.cpp
@@ -93,19 +93,21 @@ void FGoogleARCoreAndroidHelper::OnDisplayOrientationChanged()
 
 #endif
 
-int32 FGoogleARCoreAndroidHelper::CurrentDisplayRotation = 0;
+ARCoreDisplayRotation FGoogleARCoreAndroidHelper::CurrentDisplayRotation = ARCoreDisplayRotation::Rotation0;
 void FGoogleARCoreAndroidHelper::UpdateDisplayRotation()
 {
 #if PLATFORM_ANDROID
 	if (JNIEnv* Env = FAndroidApplication::GetJavaEnv())
 	{
 		static jmethodID Method = FJavaWrapper::FindMethod(Env, FJavaWrapper::GameActivityClassID, "AndroidThunkJava_GetDisplayRotation", "()I", false);
-		CurrentDisplayRotation = FJavaWrapper::CallIntMethod(Env, FJavaWrapper::GameActivityThis, Method);
+		int32 IntCurrentDisplayRotation = FJavaWrapper::CallIntMethod(Env, FJavaWrapper::GameActivityThis, Method);
+		CurrentDisplayRotation = static_cast<ARCoreDisplayRotation>(IntCurrentDisplayRotation);
+		check((CurrentDisplayRotation >= ARCoreDisplayRotation::Rotation0) || (CurrentDisplayRotation <= ARCoreDisplayRotation::Max));
 	}
 #endif
 }
 
-int32 FGoogleARCoreAndroidHelper::GetDisplayRotation()
+ARCoreDisplayRotation FGoogleARCoreAndroidHelper::GetDisplayRotation()
 {
 	return CurrentDisplayRotation;
 }
diff --git a/Engine/Plugins/Runtime/GoogleARCore/Source/GoogleARCoreBase/Private/GoogleARCoreAndroidHelper.h b/Engine/Plugins/Runtime/GoogleARCore/Source/GoogleARCoreBase/Private/GoogleARCoreAndroidHelper.h
--- a/Engine/Plugins/Runtime/GoogleARCore/Source/GoogleARCoreBase/Private/GoogleARCoreAndroidHelper.h
+++ b/Engine/Plugins/Runtime/GoogleARCore/Source/GoogleARCoreBase/Private/GoogleARCoreAndroidHelper.h
@@ -7,6 +7,16 @@
 #include "Android/AndroidApplication.h"
 #endif
 
+// Matches return values from android.view.display.getRotation()
+enum class ARCoreDisplayRotation : int32
+{
+	Rotation0 = 0,
+	Rotation90 = 1,
+	Rotation180 = 2,
+	Rotation270 = 3,
+	Max = 3
+};
+
 /** Wrappers for accessing Android Java stuff */
 class FGoogleARCoreAndroidHelper
 {
@@ -20,7 +30,7 @@ public:
 	/**
 	 * Get Andriod display orientation.
 	 */
-	static int32 GetDisplayRotation();
+	static ARCoreDisplayRotation GetDisplayRotation();
 
 	static void QueueStartSessionOnUiThread();
 
@@ -36,5 +46,5 @@ public:
 #endif
 
 private:
-	static int32 CurrentDisplayRotation;
+	static ARCoreDisplayRotation CurrentDisplayRotation;
 };
diff --git a/Engine/Plugins/Runtime/GoogleARCore/Source/GoogleARCoreBase/Private/GoogleARCoreDevice.cpp b/Engine/Plugins/Runtime/GoogleARCore/Source/GoogleARCoreBase/Private/GoogleARCoreDevice.cpp
--- a/Engine/Plugins/Runtime/GoogleARCore/Source/GoogleARCoreBase/Private/GoogleARCoreDevice.cpp
+++ b/Engine/Plugins/Runtime/GoogleARCore/Source/GoogleARCoreBase/Private/GoogleARCoreDevice.cpp
@@ -249,7 +249,7 @@ void FGoogleARCoreDevice::OnWorldTickStart(ELevelTick TickType, float DeltaTime)
 		{
 			ViewportSize = GEngine->GameViewport->Viewport->GetSizeXY();
 		}
-		ARCoreSession->SetDisplayGeometry(FGoogleARCoreAndroidHelper::GetDisplayRotation(), ViewportSize.X, ViewportSize.Y);
+		ARCoreSession->SetDisplayGeometry(static_cast<int>(FGoogleARCoreAndroidHelper::GetDisplayRotation()), ViewportSize.X, ViewportSize.Y);
 		EGoogleARCoreAPIStatus Status = ARCoreSession->Update(WorldToMeterScale);
 		if (Status == EGoogleARCoreAPIStatus::AR_ERROR_FATAL)
 		{
diff --git a/Engine/Plugins/Runtime/GoogleARCore/Source/GoogleARCoreBase/Private/GoogleARCoreXRCamera.cpp b/Engine/Plugins/Runtime/GoogleARCore/Source/GoogleARCoreBase/Private/GoogleARCoreXRCamera.cpp
--- a/Engine/Plugins/Runtime/GoogleARCore/Source/GoogleARCoreBase/Private/GoogleARCoreXRCamera.cpp
+++ b/Engine/Plugins/Runtime/GoogleARCore/Source/GoogleARCoreBase/Private/GoogleARCoreXRCamera.cpp
@@ -4,6 +4,7 @@
 #include "GoogleARCoreXRTrackingSystem.h"
 #include "SceneView.h"
 #include "GoogleARCorePassthroughCameraRenderer.h"
+#include "GoogleARCoreAndroidHelper.h"
 
 FGoogleARCoreXRCamera::FGoogleARCoreXRCamera(const FAutoRegister& AutoRegister, FGoogleARCoreXRTrackingSystem& InARCoreSystem, int32 InDeviceID)
 	: FDefaultXRCamera(AutoRegister, &InARCoreSystem, InDeviceID)
@@ -53,8 +54,7 @@ void FGoogleARCoreXRCamera::PostRenderBasePass_RenderThread(FRHICommandListImmed
 	{
 		TArray<float> TransformedUVs;
 		GoogleARCoreTrackingSystem.ARCoreDeviceInstance->GetPassthroughCameraImageUVs(PassthroughRenderer->OverlayQuadUVs, TransformedUVs);
-
-		PassthroughRenderer->UpdateOverlayUVCoordinate_RenderThread(TransformedUVs);
+		PassthroughRenderer->UpdateOverlayUVCoordinate_RenderThread(TransformedUVs, FGoogleARCoreAndroidHelper::GetDisplayRotation());
 		PassthroughRenderer->RenderVideoOverlay_RenderThread(RHICmdList, InView);
 	}
 }
diff --git a/Engine/Plugins/Runtime/GoogleARCore/Source/GoogleARCoreRendering/GoogleARCoreRendering.Build.cs b/Engine/Plugins/Runtime/GoogleARCore/Source/GoogleARCoreRendering/GoogleARCoreRendering.Build.cs
--- a/Engine/Plugins/Runtime/GoogleARCore/Source/GoogleARCoreRendering/GoogleARCoreRendering.Build.cs
+++ b/Engine/Plugins/Runtime/GoogleARCore/Source/GoogleARCoreRendering/GoogleARCoreRendering.Build.cs
@@ -9,6 +9,7 @@ public class GoogleARCoreRendering : ModuleRules
 		PrivateIncludePaths.AddRange(
 			new string[] {
 				"GoogleARCoreRendering/Private",
+				"GoogleARCoreBase/Private",
 				"../../../../Source/Runtime/Renderer/Private",
 				// ... add other private include paths required here ...
 			}
diff --git a/Engine/Plugins/Runtime/GoogleARCore/Source/GoogleARCoreRendering/Private/GoogleARCorePassthroughCameraRenderer.cpp b/Engine/Plugins/Runtime/GoogleARCore/Source/GoogleARCoreRendering/Private/GoogleARCorePassthroughCameraRenderer.cpp
--- a/Engine/Plugins/Runtime/GoogleARCore/Source/GoogleARCoreRendering/Private/GoogleARCorePassthroughCameraRenderer.cpp
+++ b/Engine/Plugins/Runtime/GoogleARCore/Source/GoogleARCoreRendering/Private/GoogleARCorePassthroughCameraRenderer.cpp
@@ -17,6 +17,7 @@
 #include "MaterialShaderType.h"
 #include "ExternalTexture.h"
 #include "GoogleARCorePassthroughCameraExternalTextureGuid.h"
+#include "GoogleARCoreAndroidHelper.h"
 
 FGoogleARCorePassthroughCameraRenderer::FGoogleARCorePassthroughCameraRenderer()
 	: OverlayQuadUVs{ 0.0f, 0.0f, 0.0f, 1.0f, 1.0f, 0.0f, 1.0f, 1.0f }
@@ -96,20 +97,37 @@ void FGoogleARCorePassthroughCameraRenderer::InitializeRenderer_RenderThread(FTe
 	bInitialized = true;
 }
 
-void FGoogleARCorePassthroughCameraRenderer::UpdateOverlayUVCoordinate_RenderThread(TArray<float>& InOverlayUVs)
+void FGoogleARCorePassthroughCameraRenderer::UpdateOverlayUVCoordinate_RenderThread(TArray<float>& InOverlayUVs, ARCoreDisplayRotation DisplayRotation)
 {
 	check(InOverlayUVs.Num() == 8);
 
-	bool bNeedToFlipCameraImage = RHINeedsToSwitchVerticalAxis(GShaderPlatformForFeatureLevel[GMaxRHIFeatureLevel]) && !IsMobileHDR();
+	// It seems very likely that this is papering over up some underlying problem with the camera image orientation.
+	bool bIsLandscape = (DisplayRotation == ARCoreDisplayRotation::Rotation0 || DisplayRotation == ARCoreDisplayRotation::Rotation180);
+	bool bNeedToFlipCameraImageHorizontally = !bIsLandscape && RHINeedsToSwitchVerticalAxis(GShaderPlatformForFeatureLevel[GMaxRHIFeatureLevel]) && !IsMobileHDR();
+	bool bFlipCameraImageVertically = bIsLandscape && RHINeedsToSwitchVerticalAxis(GShaderPlatformForFeatureLevel[GMaxRHIFeatureLevel]) && !IsMobileHDR();
+	bool bDiagonalFlip = bIsLandscape && IsMobileHDR();
 
-	if (bNeedToFlipCameraImage)
+	if (bFlipCameraImageVertically)
+	{
+		InOverlayUVs.SwapMemory(0, 4);
+		InOverlayUVs.SwapMemory(1, 5);
+		InOverlayUVs.SwapMemory(2, 6);
+		InOverlayUVs.SwapMemory(3, 7);
+	}
+	else if (bNeedToFlipCameraImageHorizontally)
 	{
 		InOverlayUVs.SwapMemory(0, 2);
 		InOverlayUVs.SwapMemory(1, 3);
 		InOverlayUVs.SwapMemory(4, 6);
 		InOverlayUVs.SwapMemory(5, 7);
 	}
-
+	else if (bDiagonalFlip)
+	{
+		InOverlayUVs.SwapMemory(0, 6);
+		InOverlayUVs.SwapMemory(1, 7);
+		InOverlayUVs.SwapMemory(4, 2);
+		InOverlayUVs.SwapMemory(5, 3);
+	}
 	{
 		if (OverlayVertexBufferRHI.IsValid())
 		{
@@ -136,40 +154,6 @@ void FGoogleARCorePassthroughCameraRenderer::UpdateOverlayUVCoordinate_RenderThr
 		FRHIResourceCreateInfo CreateInfo(&Vertices);
 		OverlayVertexBufferRHI = RHICreateVertexBuffer(Vertices.GetResourceDataSize(), BUF_Static, CreateInfo);
 	}
-
-	// Force swap vertical axis for the copy to texture UV
-	if (!bNeedToFlipCameraImage)
-	{
-		InOverlayUVs.SwapMemory(0, 2);
-		InOverlayUVs.SwapMemory(1, 3);
-		InOverlayUVs.SwapMemory(4, 6);
-		InOverlayUVs.SwapMemory(5, 7);
-	}
-	{
-		if (OverlayCopyVertexBufferRHI.IsValid())
-		{
-			OverlayCopyVertexBufferRHI.SafeRelease();
-		}
-
-		TResourceArray<FFilterVertex, VERTEXBUFFER_ALIGNMENT> Vertices;
-		Vertices.SetNumUninitialized(4);
-
-		Vertices[0].Position = FVector4(0, 1, 0, 1);
-		Vertices[0].UV = FVector2D(InOverlayUVs[0], InOverlayUVs[1]);
-
-		Vertices[1].Position = FVector4(0, 0, 0, 1);
-		Vertices[1].UV = FVector2D(InOverlayUVs[2], InOverlayUVs[3]);
-
-		Vertices[2].Position = FVector4(1, 1, 0, 1);
-		Vertices[2].UV = FVector2D(InOverlayUVs[4], InOverlayUVs[5]);
-
-		Vertices[3].Position = FVector4(1, 0, 0, 1);
-		Vertices[3].UV = FVector2D(InOverlayUVs[6], InOverlayUVs[7]);
-
-		// Create vertex buffer. Fill buffer with initial data upon creation
-		FRHIResourceCreateInfo CreateInfo(&Vertices);
-		OverlayCopyVertexBufferRHI = RHICreateVertexBuffer(Vertices.GetResourceDataSize(), BUF_Static, CreateInfo);
-	}
 }
 
 // We use something similar to the PostProcessMaterial to render the color camera overlay.
@@ -341,48 +325,3 @@ void FGoogleARCorePassthroughCameraRenderer::RenderVideoOverlay_RenderThread(FRH
 	}
 #endif
 }
-
-void FGoogleARCorePassthroughCameraRenderer::CopyVideoTexture_RenderThread(FRHICommandListImmediate& RHICmdList, FTextureRHIParamRef DstTexture, FIntPoint TargetSize)
-{
-#if PLATFORM_ANDROID
-	const auto FeatureLevel = GMaxRHIFeatureLevel;
-	auto ShaderMap = GetGlobalShaderMap(FeatureLevel);
-
-	TShaderMapRef<FScreenVS> VertexShader(ShaderMap);
-	TShaderMapRef<FScreenPS_OSE> PixelShader(ShaderMap);
-
-	IRendererModule& RendererModule = GetRendererModule();
-
-	SetRenderTarget(RHICmdList, DstTexture, FTextureRHIRef());
-	DrawClearQuad(RHICmdList, FLinearColor(0.0f, 0.0f, 0.0f, 0.0f));
-
-	RHICmdList.SetViewport(0, 0, 0, TargetSize.X, TargetSize.Y, 1.0f);
-
-	FGraphicsPipelineStateInitializer GraphicsPSOInit;
-	RHICmdList.ApplyCachedRenderTargets(GraphicsPSOInit);
-
-	GraphicsPSOInit.BlendState = TStaticBlendState<>::GetRHI();
-	GraphicsPSOInit.RasterizerState = TStaticRasterizerState<>::GetRHI();
-	GraphicsPSOInit.DepthStencilState = TStaticDepthStencilState<false, CF_Always>::GetRHI();
-
-	GraphicsPSOInit.BoundShaderState.VertexDeclarationRHI = RendererModule.GetFilterVertexDeclaration().VertexDeclarationRHI;
-	GraphicsPSOInit.BoundShaderState.VertexShaderRHI = GETSAFERHISHADER_VERTEX(*VertexShader);
-	GraphicsPSOInit.BoundShaderState.PixelShaderRHI = GETSAFERHISHADER_PIXEL(*PixelShader);
-	GraphicsPSOInit.PrimitiveType = PT_TriangleList;
-
-	SetGraphicsPipelineState(RHICmdList, GraphicsPSOInit);
-
-	PixelShader->SetParameters(RHICmdList, TStaticSamplerState<SF_Bilinear>::GetRHI(), VideoTexture.GetReference());
-
-	RendererModule.DrawRectangle(
-		RHICmdList,
-		0, 0,
-		TargetSize.X, TargetSize.Y,
-		0, 0,
-		1, 1,
-		FIntPoint(TargetSize.X, TargetSize.Y),
-		FIntPoint(1, 1),
-		*VertexShader,
-		EDRF_Default);
-#endif
-}
diff --git a/Engine/Plugins/Runtime/GoogleARCore/Source/GoogleARCoreRendering/Public/GoogleARCorePassthroughCameraRenderer.h b/Engine/Plugins/Runtime/GoogleARCore/Source/GoogleARCoreRendering/Public/GoogleARCorePassthroughCameraRenderer.h
--- a/Engine/Plugins/Runtime/GoogleARCore/Source/GoogleARCoreRendering/Public/GoogleARCorePassthroughCameraRenderer.h
+++ b/Engine/Plugins/Runtime/GoogleARCore/Source/GoogleARCoreRendering/Public/GoogleARCorePassthroughCameraRenderer.h
@@ -17,6 +17,8 @@
 
 #include "GoogleARCorePassthroughCameraRenderer.generated.h"
 
+enum class ARCoreDisplayRotation : int32;
+
 /** A helper class that is used to load the GoogleARCorePassthroughCameraMaterial from its default object. */
 UCLASS()
 class GOOGLEARCORERENDERING_API UGoogleARCoreCameraOverlayMaterialLoader : public UObject
@@ -56,16 +58,14 @@ public:
 
 	void InitializeRenderer_RenderThread(FTextureRHIRef ExternalTexture);
 
-	void UpdateOverlayUVCoordinate_RenderThread(TArray<float>& InOverlayUVs);
+	void UpdateOverlayUVCoordinate_RenderThread(TArray<float>& InOverlayUVs, ARCoreDisplayRotation DisplayRotation);
 	void RenderVideoOverlay_RenderThread(FRHICommandListImmediate& RHICmdList, FSceneView& InView);
-	void CopyVideoTexture_RenderThread(FRHICommandListImmediate& RHICmdList, FTextureRHIParamRef DstTexture, FIntPoint TargetSize);
 
 	const TArray<float> OverlayQuadUVs;
 private:
 	bool bInitialized;
 	FIndexBufferRHIRef OverlayIndexBufferRHI;
 	FVertexBufferRHIRef OverlayVertexBufferRHI;
-	FVertexBufferRHIRef OverlayCopyVertexBufferRHI;
 	FTextureRHIRef VideoTexture;
 	float OverlayTextureUVs[8];
 	bool bMaterialInitialized;
