Fix native JNI calls in GoogleARCore

#ue4
#android

diff --git a/Engine/Plugins/Runtime/GoogleARCore/Source/GoogleARCoreBase/Private/GoogleARCoreAndroidHelper.cpp b/Engine/Plugins/Runtime/GoogleARCore/Source/GoogleARCoreBase/Private/GoogleARCoreAndroidHelper.cpp
--- a/Engine/Plugins/Runtime/GoogleARCore/Source/GoogleARCoreBase/Private/GoogleARCoreAndroidHelper.cpp
+++ b/Engine/Plugins/Runtime/GoogleARCore/Source/GoogleARCoreBase/Private/GoogleARCoreAndroidHelper.cpp
@@ -13,42 +13,42 @@ extern "C"
 {
 	// Functions that are called on Android lifecycle events.
 
-	void Java_com_google_arcore_unreal_GoogleARCoreJavaHelper_onApplicationCreated(JNIEnv*, jobject)
+	JNI_METHOD void Java_com_google_arcore_unreal_GoogleARCoreJavaHelper_onApplicationCreated(JNIEnv*, jobject)
 	{
 		FGoogleARCoreAndroidHelper::OnApplicationCreated();
 	}
 
-	void Java_com_google_arcore_unreal_GoogleARCoreJavaHelper_onApplicationDestroyed(JNIEnv*, jobject)
+	JNI_METHOD void Java_com_google_arcore_unreal_GoogleARCoreJavaHelper_onApplicationDestroyed(JNIEnv*, jobject)
 	{
 		FGoogleARCoreAndroidHelper::OnApplicationDestroyed();
 	}
 
-	void Java_com_google_arcore_unreal_GoogleARCoreJavaHelper_onApplicationPause(JNIEnv*, jobject)
+	JNI_METHOD void Java_com_google_arcore_unreal_GoogleARCoreJavaHelper_onApplicationPause(JNIEnv*, jobject)
 	{
 		FGoogleARCoreAndroidHelper::OnApplicationPause();
 	}
 
-	void Java_com_google_arcore_unreal_GoogleARCoreJavaHelper_onApplicationResume(JNIEnv*, jobject)
+	JNI_METHOD void Java_com_google_arcore_unreal_GoogleARCoreJavaHelper_onApplicationResume(JNIEnv*, jobject)
 	{
 		FGoogleARCoreAndroidHelper::OnApplicationResume();
 	}
 
-	void Java_com_google_arcore_unreal_GoogleARCoreJavaHelper_onApplicationStop(JNIEnv*, jobject)
+	JNI_METHOD void Java_com_google_arcore_unreal_GoogleARCoreJavaHelper_onApplicationStop(JNIEnv*, jobject)
 	{
 		FGoogleARCoreAndroidHelper::OnApplicationStop();
 	}
 
-	void Java_com_google_arcore_unreal_GoogleARCoreJavaHelper_onApplicationStart(JNIEnv*, jobject)
+	JNI_METHOD void Java_com_google_arcore_unreal_GoogleARCoreJavaHelper_onApplicationStart(JNIEnv*, jobject)
 	{
 		FGoogleARCoreAndroidHelper::OnApplicationStart();
 	}
 
-	void Java_com_google_arcore_unreal_GoogleARCoreJavaHelper_onDisplayOrientationChanged(JNIEnv*, jobject)
+	JNI_METHOD void Java_com_google_arcore_unreal_GoogleARCoreJavaHelper_onDisplayOrientationChanged(JNIEnv*, jobject)
 	{
 		FGoogleARCoreAndroidHelper::OnDisplayOrientationChanged();
 	}
 
-	void Java_com_google_arcore_unreal_GoogleARCoreJavaHelper_ARCoreSessionStart(JNIEnv*, jobject)
+	JNI_METHOD void Java_com_google_arcore_unreal_GoogleARCoreJavaHelper_ARCoreSessionStart(JNIEnv*, jobject)
 	{
 		FGoogleARCoreDevice::GetInstance()->StartSessionWithRequestedConfig();
 	}
