Mark GoogleARCore as only supporting Android as a target platform, so it's not precompiled for all desktop platforms as part of the installed UE4 build. Fixes CIS.

diff --git a/Engine/Plugins/Runtime/GoogleARCore/GoogleARCore.uplugin b/Engine/Plugins/Runtime/GoogleARCore/GoogleARCore.uplugin
--- a/Engine/Plugins/Runtime/GoogleARCore/GoogleARCore.uplugin
+++ b/Engine/Plugins/Runtime/GoogleARCore/GoogleARCore.uplugin
@@ -11,19 +11,20 @@
     "CanContainContent": true,
     "IsBetaVersion": false,
     "EnabledByDefault": false,
+	"SupportedTargetPlatforms": [
+		"Android"
+	],
     "Modules" :
     [
         {
             "Name" : "GoogleARCoreBase",
-            "Type" : "Runtime",
-            "WhitelistPlatforms": [ "Win64", "Mac", "Android", "Linux" ]
+            "Type" : "Runtime"
         },
         {
             "Name": "GoogleARCoreRendering",
             "Type": "Runtime",
-            "LoadingPhase" : "PostConfigInit",
-            "WhitelistPlatforms": [ "Win64", "Mac", "Android", "Linux" ]
-        },
+            "LoadingPhase" : "PostConfigInit"
+        }
     ],
     "Plugins": [
         {
