Fix for ARCore Tools being in the wrong folder, emitting RuntimeDependency ref to make sure they're properly included

diff --git a/Engine/Plugins/Runtime/GoogleARCore/Source/GoogleARCoreBase/GoogleARCoreBase.Build.cs b/Engine/Plugins/Runtime/GoogleARCore/Source/GoogleARCoreBase/GoogleARCoreBase.Build.cs
--- a/Engine/Plugins/Runtime/GoogleARCore/Source/GoogleARCoreBase/GoogleARCoreBase.Build.cs
+++ b/Engine/Plugins/Runtime/GoogleARCore/Source/GoogleARCoreBase/GoogleARCoreBase.Build.cs
@@ -67,7 +67,29 @@ namespace UnrealBuildTool.Rules
 				}
 			}
 
+			if (Target.bBuildEditor)
+			{
+				string ExecName = "";
+				if (Target.Platform == UnrealTargetPlatform.Win64)
+				{
+					ExecName = "arcoreimg.exe";
+				}
+				else if (Target.Platform == UnrealTargetPlatform.Linux)
+				{
+					ExecName = "arcoreimg";
+				}
+				else if (Target.Platform == UnrealTargetPlatform.Mac)
+				{
+					ExecName = "ptdbtool_macos_lipobin";
+				}
+				
+				if (ExecName.Length > 0)
+				{
+					RuntimeDependencies.Add("$(EngineDir)/Plugins/Runtime/GoogleARCore/Binaries/ThirdParty/Google/ARCoreImg/" + Target.Platform.ToString() + ExecName);
+				}
+			}
+			
 			bFasterWithoutUnity = false;
-		}
-	}
+        }
+    }
 }
diff --git a/Engine/Plugins/Runtime/GoogleARCore/Source/GoogleARCoreBase/Private/GoogleARCoreAugmentedImageDatabase.cpp b/Engine/Plugins/Runtime/GoogleARCore/Source/GoogleARCoreBase/Private/GoogleARCoreAugmentedImageDatabase.cpp
--- a/Engine/Plugins/Runtime/GoogleARCore/Source/GoogleARCoreBase/Private/GoogleARCoreAugmentedImageDatabase.cpp
+++ b/Engine/Plugins/Runtime/GoogleARCore/Source/GoogleARCoreBase/Private/GoogleARCoreAugmentedImageDatabase.cpp
@@ -159,7 +159,10 @@ void UGoogleARCoreAugmentedImageDatabase::Serialize(FArchive& Ar)
 					*FPaths::EnginePluginsDir(),
 					TEXT("Runtime"),
 					TEXT("GoogleARCore"),
-					TEXT("Tools"),
+					TEXT("Binaries"),
+					TEXT("ThirdParty"),
+					TEXT("Google"),
+					TEXT("ARCoreImg"),
 					*UGameplayStatics::GetPlatformName(),
 #if PLATFORM_LINUX
 					TEXT("arcoreimg")
