Merging fix from CL 4247094 for ARCore external dependency locations

diff --git a/Engine/Plugins/Runtime/GoogleARCore/Source/GoogleARCoreBase/GoogleARCoreBase.Build.cs b/Engine/Plugins/Runtime/GoogleARCore/Source/GoogleARCoreBase/GoogleARCoreBase.Build.cs
--- a/Engine/Plugins/Runtime/GoogleARCore/Source/GoogleARCoreBase/GoogleARCoreBase.Build.cs
+++ b/Engine/Plugins/Runtime/GoogleARCore/Source/GoogleARCoreBase/GoogleARCoreBase.Build.cs
@@ -72,20 +72,20 @@ namespace UnrealBuildTool.Rules
 				string ExecName = "";
 				if (Target.Platform == UnrealTargetPlatform.Win64)
 				{
-					ExecName = "arcoreimg.exe";
+					ExecName = "Windows/arcoreimg.exe";
 				}
 				else if (Target.Platform == UnrealTargetPlatform.Linux)
 				{
-					ExecName = "arcoreimg";
+					ExecName = "Linux/arcoreimg";
 				}
 				else if (Target.Platform == UnrealTargetPlatform.Mac)
 				{
-					ExecName = "ptdbtool_macos_lipobin";
+					ExecName = "Mac/ptdbtool_macos_lipobin";
 				}
 				
 				if (ExecName.Length > 0)
 				{
-					RuntimeDependencies.Add("$(EngineDir)/Plugins/Runtime/GoogleARCore/Binaries/ThirdParty/Google/ARCoreImg/" + Target.Platform.ToString() + ExecName);
+					RuntimeDependencies.Add("$(EngineDir)/Plugins/Runtime/GoogleARCore/Binaries/ThirdParty/Google/ARCoreImg/" + ExecName);
 				}
 			}
 			
