Merging //UE4/Dev-Main@4569249 to Dev-Rendering (//UE4/Dev-Rendering)

diff --git a/Engine/Plugins/Runtime/GoogleARCore/GoogleARCore.uplugin b/Engine/Plugins/Runtime/GoogleARCore/GoogleARCore.uplugin
--- a/Engine/Plugins/Runtime/GoogleARCore/GoogleARCore.uplugin
+++ b/Engine/Plugins/Runtime/GoogleARCore/GoogleARCore.uplugin
@@ -11,19 +11,20 @@
     "CanContainContent": true,
     "IsBetaVersion": false,
     "EnabledByDefault": false,
+	"SupportedTargetPlatforms": [
+		"Android"
+	],
     "Modules" :
     [
         {
             "Name" : "GoogleARCoreBase",
-            "Type" : "Runtime",
-            "WhitelistPlatforms": [ "Win64", "Mac", "Android", "Linux" ]
+            "Type" : "Runtime"
         },
         {
             "Name": "GoogleARCoreRendering",
             "Type": "Runtime",
-            "LoadingPhase" : "PostConfigInit",
-            "WhitelistPlatforms": [ "Win64", "Mac", "Android", "Linux" ]
-        },
+            "LoadingPhase" : "PostConfigInit"
+        }
     ],
     "Plugins": [
         {
diff --git a/Engine/Source/Programs/AutomationTool/Scripts/CheckForHacks.cs b/Engine/Source/Programs/AutomationTool/Scripts/CheckForHacks.cs
--- a/Engine/Source/Programs/AutomationTool/Scripts/CheckForHacks.cs
+++ b/Engine/Source/Programs/AutomationTool/Scripts/CheckForHacks.cs
@@ -48,7 +48,7 @@ namespace AutomationTool
 		/// <summary>
 		/// The pattern that should match hack comments (and captures a list of tags). Ignore anything with a semicolon between BEGIN and HACK to avoid matching statements with "HACK" as a comment (seen in LLVM).
 		/// </summary>
-		static readonly Regex CompiledRegex = new Regex("(?<!\\w)BEGIN\\s([^;]*)\\sHACK(?!\\w)", RegexOptions.IgnoreCase | RegexOptions.Compiled);
+		static readonly Regex CompiledRegex = new Regex("(?<!\\w)(?:BEGIN|START)\\s([^;]*)\\sHACK(?!\\w)", RegexOptions.IgnoreCase | RegexOptions.Compiled);
 
 		/// <summary>
 		/// Executes the command
