Add API to return passthrough camera UVs. (Currently only implemented by GoogleARCore)

diff --git a/Engine/Plugins/Runtime/AR/Google/GoogleARCore/Source/GoogleARCoreBase/Private/GoogleARCoreXRCamera.cpp b/Engine/Plugins/Runtime/AR/Google/GoogleARCore/Source/GoogleARCoreBase/Private/GoogleARCoreXRCamera.cpp
--- a/Engine/Plugins/Runtime/AR/Google/GoogleARCore/Source/GoogleARCoreBase/Private/GoogleARCoreXRCamera.cpp
+++ b/Engine/Plugins/Runtime/AR/Google/GoogleARCore/Source/GoogleARCoreBase/Private/GoogleARCoreXRCamera.cpp
@@ -49,13 +49,32 @@ void FGoogleARCoreXRCamera::PreRenderViewFamily_RenderThread(FRHICommandListImme
 }
 
 void FGoogleARCoreXRCamera::PostRenderBasePass_RenderThread(FRHICommandListImmediate& RHICmdList, FSceneView& InView)
+{
+	TArray<FVector2D> Tmp;
+	if (GetPassthroughCameraUVs_RenderThread(Tmp))
+	{
+		PassthroughRenderer->RenderVideoOverlay_RenderThread(RHICmdList, InView);
+	}
+}
+
+bool FGoogleARCoreXRCamera::GetPassthroughCameraUVs_RenderThread(TArray<FVector2D>& OutUVs)
 {
 	if (GoogleARCoreTrackingSystem.ARCoreDeviceInstance->GetIsARCoreSessionRunning() && bEnablePassthroughCameraRendering_RT)
 	{
 		TArray<float> TransformedUVs;
+		// TODO save the transformed UVs and only calculate if uninitialized or FGoogleARCoreFrame::IsDisplayRotationChanged() returns true
 		GoogleARCoreTrackingSystem.ARCoreDeviceInstance->GetPassthroughCameraImageUVs(PassthroughRenderer->OverlayQuadUVs, TransformedUVs);
 		PassthroughRenderer->UpdateOverlayUVCoordinate_RenderThread(TransformedUVs, FGoogleARCoreAndroidHelper::GetDisplayRotation());
-		PassthroughRenderer->RenderVideoOverlay_RenderThread(RHICmdList, InView);
+		OutUVs.SetNumUninitialized(4);
+		OutUVs[0] = FVector2D(TransformedUVs[0], TransformedUVs[1]);
+		OutUVs[1] = FVector2D(TransformedUVs[2], TransformedUVs[3]);
+		OutUVs[2] = FVector2D(TransformedUVs[4], TransformedUVs[5]);
+		OutUVs[3] = FVector2D(TransformedUVs[6], TransformedUVs[7]);
+		return true;
+	}
+	else
+	{
+		return false;
 	}
 }
 
diff --git a/Engine/Plugins/Runtime/AR/Google/GoogleARCore/Source/GoogleARCoreBase/Private/GoogleARCoreXRCamera.h b/Engine/Plugins/Runtime/AR/Google/GoogleARCore/Source/GoogleARCoreBase/Private/GoogleARCoreXRCamera.h
--- a/Engine/Plugins/Runtime/AR/Google/GoogleARCore/Source/GoogleARCoreBase/Private/GoogleARCoreXRCamera.h
+++ b/Engine/Plugins/Runtime/AR/Google/GoogleARCore/Source/GoogleARCoreBase/Private/GoogleARCoreXRCamera.h
@@ -21,6 +21,7 @@ public:
 	virtual void PreRenderViewFamily_RenderThread(FRHICommandListImmediate& RHICmdList, FSceneViewFamily& InViewFamily) override;
 	virtual void PostRenderBasePass_RenderThread(FRHICommandListImmediate& RHICmdList, FSceneView& InView) override;
 	virtual bool IsActiveThisFrame(class FViewport* InViewport) const override;
+	virtual bool GetPassthroughCameraUVs_RenderThread(TArray<FVector2D>& OutUVs) override;
 	//~ FDefaultXRCamera
 
 	void ConfigXRCamera(bool bInMatchDeviceCameraFOV, bool bInEnablePassthroughCameraRendering);
diff --git a/Engine/Source/Runtime/HeadMountedDisplay/Public/IXRCamera.h b/Engine/Source/Runtime/HeadMountedDisplay/Public/IXRCamera.h
--- a/Engine/Source/Runtime/HeadMountedDisplay/Public/IXRCamera.h
+++ b/Engine/Source/Runtime/HeadMountedDisplay/Public/IXRCamera.h
@@ -59,5 +59,12 @@ public:
 	 */
 	virtual void CalculateStereoCameraOffset(const enum EStereoscopicPass StereoPassType, FRotator& ViewRotation, FVector& ViewLocation) = 0;
 
-
+	/**
+	 * Fetches the UV coordinates of an AR passthrough camera relative to the screen. Takes into account cropping, orientation and difference between
+	 * camera and screen aspect ratios. Returns false if the UVs are not available or if the feature is not supported or implemented by the current platform.
+	 * If the result is true, OutUVs will contain the four corner UVs of the passthrough camera rectangle.
+	 *
+	 * This method should only be called from the render thread.
+	 */
+	virtual bool GetPassthroughCameraUVs_RenderThread(TArray<FVector2D>& OutUVs) { return false; }
 };
