Fix from ARCore changes. (Renable when missing support makes it to main)

diff --git a/Engine/Plugins/Runtime/AR/Google/GoogleARCore/Source/GoogleARCoreBase/Private/GoogleARCoreXRTrackingSystem.cpp b/Engine/Plugins/Runtime/AR/Google/GoogleARCore/Source/GoogleARCoreBase/Private/GoogleARCoreXRTrackingSystem.cpp
--- a/Engine/Plugins/Runtime/AR/Google/GoogleARCore/Source/GoogleARCoreBase/Private/GoogleARCoreXRTrackingSystem.cpp
+++ b/Engine/Plugins/Runtime/AR/Google/GoogleARCore/Source/GoogleARCoreBase/Private/GoogleARCoreXRTrackingSystem.cpp
@@ -199,6 +199,7 @@ EARTrackingQuality FGoogleARCoreXRTrackingSystem::OnGetTrackingQuality() const
 	return EARTrackingQuality::OrientationAndPosition;
 }
 
+/*
 EARTrackingQualityReason FGoogleARCoreXRTrackingSystem::OnGetTrackingQualityReason() const
 {
 	// Dont return EARTrackingQualityReason::None, which means that the tracking quality is not limited
@@ -209,6 +210,7 @@ bool FGoogleARCoreXRTrackingSystem::IsARAvailable() const
 {
 	return FGoogleARCoreDevice::GetInstance()->CheckARCoreAPKAvailability() == EGoogleARCoreAvailability::SupportedInstalled;
 }
+*/
 
 void FGoogleARCoreXRTrackingSystem::OnStartARSession(UARSessionConfig* SessionConfig)
 {
diff --git a/Engine/Plugins/Runtime/AR/Google/GoogleARCore/Source/GoogleARCoreBase/Private/GoogleARCoreXRTrackingSystem.h b/Engine/Plugins/Runtime/AR/Google/GoogleARCore/Source/GoogleARCoreBase/Private/GoogleARCoreXRTrackingSystem.h
--- a/Engine/Plugins/Runtime/AR/Google/GoogleARCore/Source/GoogleARCoreBase/Private/GoogleARCoreXRTrackingSystem.h
+++ b/Engine/Plugins/Runtime/AR/Google/GoogleARCore/Source/GoogleARCoreBase/Private/GoogleARCoreXRTrackingSystem.h
@@ -48,11 +48,12 @@ public:
 
 protected:
 	// IARSystemSupport
-	/** Returns true/false based on whether AR features are available */
-	virtual bool IsARAvailable() const override;
+	/** todo - reenable once base changes make their way up - Returns true/false based on whether AR features are available */
+	//virtual bool IsARAvailable() const override;
 	virtual void OnARSystemInitialized() override;
 	virtual EARTrackingQuality OnGetTrackingQuality() const override;
-	virtual EARTrackingQualityReason OnGetTrackingQualityReason() const override;
+	/** todo - reenable once base changes make their way up */
+	//virtual EARTrackingQualityReason OnGetTrackingQualityReason() const override;
 	virtual void OnStartARSession(UARSessionConfig* SessionConfig) override;
 	virtual void OnPauseARSession() override;
 	virtual void OnStopARSession() override;
