Add 0 entries to BP exposed enums in Google AR core types

#rnx

diff --git a/Engine/Plugins/Runtime/AR/Google/GoogleARCore/Source/GoogleARCoreBase/Public/GoogleARCoreTypes.h b/Engine/Plugins/Runtime/AR/Google/GoogleARCore/Source/GoogleARCoreBase/Public/GoogleARCoreTypes.h
--- a/Engine/Plugins/Runtime/AR/Google/GoogleARCore/Source/GoogleARCoreBase/Public/GoogleARCoreTypes.h
+++ b/Engine/Plugins/Runtime/AR/Google/GoogleARCore/Source/GoogleARCoreBase/Public/GoogleARCoreTypes.h
@@ -187,6 +187,7 @@ ENUM_CLASS_FLAGS(EGoogleARCoreLineTraceChannel);
 UENUM(BlueprintType, Category = "GoogleARCore|CameraConfig")
 enum class EGoogleARCoreCameraFacing : uint8
 {
+	None = 0 UMETA(Hidden),
 	/* Camera facing away from user. */
 	Back = 1,
 	/* Camera facing towards user. */
