Merging ARKit changes from dev branch

diff --git a/Engine/Plugins/Experimental/RemoteSession/Source/RemoteSession/Channels/RemoteSessionARSystemChannel.h b/Engine/Plugins/Experimental/RemoteSession/Source/RemoteSession/Channels/RemoteSessionARSystemChannel.h
--- a/Engine/Plugins/Experimental/RemoteSession/Source/RemoteSession/Channels/RemoteSessionARSystemChannel.h
+++ b/Engine/Plugins/Experimental/RemoteSession/Source/RemoteSession/Channels/RemoteSessionARSystemChannel.h
@@ -23,7 +23,13 @@ public:
 	static void Destroy();
 
 	// ~IARSystemSupport
+	
+	/** Returns true/false based on whether AR features are available */
+	virtual bool IsARAvailable() const override;
+	
 	virtual EARTrackingQuality OnGetTrackingQuality() const override;
+	/** @return the reason of limited tracking quality; if the state is not limited, return EARTrackingQualityReason::None */
+	virtual EARTrackingQualityReason OnGetTrackingQualityReason() const override;
 	virtual FARSessionStatus OnGetARSessionStatus() const override;
 	virtual TArray<UARTrackedGeometry*> OnGetAllTrackedGeometries() const override;
 	virtual bool OnIsTrackingTypeSupported(EARSessionType SessionType) const override;
diff --git a/Engine/Plugins/Experimental/RemoteSession/Source/RemoteSession/Private/Channels/RemoteSessionARSystemChannel.cpp b/Engine/Plugins/Experimental/RemoteSession/Source/RemoteSession/Private/Channels/RemoteSessionARSystemChannel.cpp
--- a/Engine/Plugins/Experimental/RemoteSession/Source/RemoteSession/Private/Channels/RemoteSessionARSystemChannel.cpp
+++ b/Engine/Plugins/Experimental/RemoteSession/Source/RemoteSession/Private/Channels/RemoteSessionARSystemChannel.cpp
@@ -57,12 +57,23 @@ void FARSystemProxy::AddReferencedObjects(FReferenceCollector& Collector)
 	Collector.AddReferencedObjects(TrackedGeometries);
 }
 
+/** Returns true/false based on whether AR features are available */
+bool FARSystemProxy::IsARAvailable() const
+{
+	return true;
+}
+
 EARTrackingQuality FARSystemProxy::OnGetTrackingQuality() const
 {
 	// @todo JoeG - Added a delegate for tracking quality so we can send that
 	return EARTrackingQuality::NotTracking;
 }
 
+EARTrackingQualityReason FARSystemProxy::OnGetTrackingQualityReason() const
+{
+	return EARTrackingQualityReason::InsufficientFeatures;
+}
+
 FARSessionStatus FARSystemProxy::OnGetARSessionStatus() const
 {
 	// @todo JoeG - Added a delegate for this so we can send that
@@ -402,4 +413,4 @@ void FRemoteSessionARSystemChannel::SendRemovedMessage(UARTrackedGeometry* Remov
 	Connection->SendPacket(Msg);
 
 	UE_LOG(LogRemoteSession, Log, TEXT("Sent trackable removed (%s)"), *Removed->GetName());
-}
\ No newline at end of file
+}
diff --git a/Engine/Plugins/Runtime/AR/Apple/AppleARKit/Source/AppleARKit/Private/AppleARKitCamera.cpp b/Engine/Plugins/Runtime/AR/Apple/AppleARKit/Source/AppleARKit/Private/AppleARKitCamera.cpp
--- a/Engine/Plugins/Runtime/AR/Apple/AppleARKit/Source/AppleARKit/Private/AppleARKitCamera.cpp
+++ b/Engine/Plugins/Runtime/AR/Apple/AppleARKit/Source/AppleARKit/Private/AppleARKitCamera.cpp
@@ -43,6 +43,31 @@ FAppleARKitCamera::FAppleARKitCamera( ARCamera* InARCamera )
 			break;
 	}
 	
+	// tracking quality reason
+	switch (InARCamera.trackingStateReason)
+	{
+		default:
+		case ARTrackingStateReasonNone:
+			TrackingQualityReason = EARTrackingQualityReason::None;
+			break;
+
+		case ARTrackingStateReasonInitializing:
+			TrackingQualityReason = EARTrackingQualityReason::Initializing;
+			break;
+
+		case ARTrackingStateReasonRelocalizing:
+			TrackingQualityReason = EARTrackingQualityReason::Relocalizing;
+			break;
+
+		case ARTrackingStateReasonExcessiveMotion:
+			TrackingQualityReason = EARTrackingQualityReason::ExcessiveMotion;
+			break;
+
+		case ARTrackingStateReasonInsufficientFeatures:
+			TrackingQualityReason = EARTrackingQualityReason::InsufficientFeatures;
+			break;
+	}
+	
 	
     // Copy / convert camera transform
     Transform = FAppleARKitConversion::ToFTransform( InARCamera.transform );
diff --git a/Engine/Plugins/Runtime/AR/Apple/AppleARKit/Source/AppleARKit/Private/AppleARKitSystem.cpp b/Engine/Plugins/Runtime/AR/Apple/AppleARKit/Source/AppleARKit/Private/AppleARKitSystem.cpp
--- a/Engine/Plugins/Runtime/AR/Apple/AppleARKit/Source/AppleARKit/Private/AppleARKitSystem.cpp
+++ b/Engine/Plugins/Runtime/AR/Apple/AppleARKit/Source/AppleARKit/Private/AppleARKitSystem.cpp
@@ -18,6 +18,7 @@
 #include "ARPin.h"
 #include "Async/Async.h"
 #include "HAL/ThreadSafeCounter.h"
+#include "HAL/IConsoleManager.h"
 #include "Misc/FileHelper.h"
 
 // For mesh occlusion
@@ -32,12 +33,17 @@
 #include "Misc/CoreDelegates.h"
 
 #if PLATFORM_IOS
-	#include "IOSRuntimeSettings.h"
-
 	#pragma clang diagnostic push
 	#pragma clang diagnostic ignored "-Wpartial-availability"
 #endif
 
+// This flag might belong in UARSessionConfig, if Apple couldn't fix the bug where
+// ARKit stops tracking surfaces after the back camera gets used by some other code
+static TAutoConsoleVariable<int32> CVarReleaseSessionWhenStopped(
+    TEXT("ar.ARKit.ReleaseSessionWhenStopped"),
+    0,
+	TEXT("Whether to release the ARKit session object when the AR session is stopped."));
+
 DECLARE_CYCLE_STAT(TEXT("SessionDidUpdateFrame_DelegateThread"), STAT_FAppleARKitSystem_SessionUpdateFrame, STATGROUP_ARKIT);
 DECLARE_CYCLE_STAT(TEXT("SessionDidAddAnchors_DelegateThread"), STAT_FAppleARKitSystem_SessionDidAddAnchors, STATGROUP_ARKIT);
 DECLARE_CYCLE_STAT(TEXT("SessionDidUpdateAnchors_DelegateThread"), STAT_FAppleARKitSystem_SessionDidUpdateAnchors, STATGROUP_ARKIT);
@@ -45,6 +51,8 @@ DECLARE_CYCLE_STAT(TEXT("SessionDidRemoveAnchors_DelegateThread"), STAT_FAppleAR
 DECLARE_CYCLE_STAT(TEXT("UpdateARKitPerf"), STAT_FAppleARKitSystem_UpdateARKitPerf, STATGROUP_ARKIT);
 DECLARE_DWORD_COUNTER_STAT(TEXT("ARKit CPU %"), STAT_ARKitThreads, STATGROUP_ARKIT);
 
+DECLARE_FLOAT_COUNTER_STAT(TEXT("ARKit Frame to Delegate Delay (ms)"), STAT_ARKitFrameToDelegateDelay, STATGROUP_ARKIT);
+
 // Copied from IOSPlatformProcess because it's not accessible by external code
 #define GAME_THREAD_PRIORITY 47
 #define RENDER_THREAD_PRIORITY 45
@@ -249,6 +257,15 @@ void FAppleARKitSystem::Shutdown()
 	PersonSegmentationDepthImage = nullptr;
 }
 
+bool FAppleARKitSystem::IsARAvailable() const
+{
+#if SUPPORTS_ARKIT_1_0
+	return FAppleARKitAvailability::SupportsARKit10();
+#else
+	return false;
+#endif
+}
+
 void FAppleARKitSystem::CheckForFaceARSupport(UARSessionConfig* InSessionConfig)
 {
 	if (InSessionConfig->GetSessionType() != EARSessionType::Face)
@@ -582,6 +599,13 @@ EARTrackingQuality FAppleARKitSystem::OnGetTrackingQuality() const
 		: EARTrackingQuality::NotTracking;
 }
 
+EARTrackingQualityReason FAppleARKitSystem::OnGetTrackingQualityReason() const
+{
+	return GameThreadFrame.IsValid()
+	? GameThreadFrame->Camera.TrackingQualityReason
+	: EARTrackingQualityReason::None;
+}
+
 void FAppleARKitSystem::OnStartARSession(UARSessionConfig* SessionConfig)
 {
 	Run(SessionConfig);
@@ -1218,7 +1242,7 @@ bool FAppleARKitSystem::HitTestAtScreenPosition(const FVector2D ScreenPosition,
 	return false;
 }
 
-void FAppleARKitSystem::SetDeviceOrientation(EDeviceScreenOrientation InOrientation)
+void FAppleARKitSystem::SetDeviceOrientationAndDerivedTracking(EDeviceScreenOrientation InOrientation)
 {
 	ensureAlwaysMsgf(InOrientation != EDeviceScreenOrientation::Unknown, TEXT("statusBarOrientation should only ever return valid orientations"));
 	if (InOrientation == EDeviceScreenOrientation::Unknown)
@@ -1227,11 +1251,11 @@ void FAppleARKitSystem::SetDeviceOrientation(EDeviceScreenOrientation InOrientat
 		InOrientation = EDeviceScreenOrientation::LandscapeLeft;
 	}
 
-	if (DeviceOrientation != InOrientation)
-	{
-		DeviceOrientation = InOrientation;
-		CalcTrackingToWorldRotation();
-	}
+	// even if this didn't change, we need to call CalcTrackingToWorldRotation because the camera mode may have changed (from Gravity to non-Gravity, etc)
+
+	DeviceOrientation = InOrientation;
+	CalcTrackingToWorldRotation();
+	
 }
 
 void FAppleARKitSystem::ClearTrackedGeometries()
@@ -1287,7 +1311,7 @@ bool FAppleARKitSystem::Run(UARSessionConfig* SessionConfig)
 	if (DeviceOrientation == EDeviceScreenOrientation::Unknown)
 	{
 		EDeviceScreenOrientation ScreenOrientation = FPlatformMisc::GetDeviceOrientation();
-		SetDeviceOrientation( ScreenOrientation );
+		SetDeviceOrientationAndDerivedTracking( ScreenOrientation );
 	}
 
 
@@ -1401,6 +1425,12 @@ bool FAppleARKitSystem::Pause()
 	{
 		// Suspend the session
 		[Session pause];
+		
+		if (CVarReleaseSessionWhenStopped.GetValueOnAnyThread())
+		{
+			[Session release];
+			Session = nullptr;
+		}
 	}
 	
 #if PLATFORM_IOS && !PLATFORM_TVOS
@@ -1429,11 +1459,16 @@ bool FAppleARKitSystem::Pause()
 void FAppleARKitSystem::OrientationChanged(const int32 NewOrientationRaw)
 {
 	const EDeviceScreenOrientation NewOrientation = static_cast<EDeviceScreenOrientation>(NewOrientationRaw);
-	SetDeviceOrientation(NewOrientation);
+	SetDeviceOrientationAndDerivedTracking(NewOrientation);
 }
 						
 void FAppleARKitSystem::SessionDidUpdateFrame_DelegateThread(TSharedPtr< FAppleARKitFrame, ESPMode::ThreadSafe > Frame)
 {
+#if STATS && SUPPORTS_ARKIT_1_0
+	const auto DelayMS = ([[NSProcessInfo processInfo] systemUptime] - Frame->Timestamp) * 1000.0;
+	SET_FLOAT_STAT(STAT_ARKitFrameToDelegateDelay, DelayMS);
+#endif
+	
 	{
 		auto UpdateFrameTask = FSimpleDelegateGraphTask::FDelegate::CreateThreadSafeSP( this, &FAppleARKitSystem::SessionDidUpdateFrame_Internal, Frame.ToSharedRef() );
 		FSimpleDelegateGraphTask::CreateAndDispatchWhenReady(UpdateFrameTask, GET_STATID(STAT_FAppleARKitSystem_SessionUpdateFrame), nullptr, ENamedThreads::GameThread);
diff --git a/Engine/Plugins/Runtime/AR/Apple/AppleARKit/Source/AppleARKit/Private/AppleARKitSystem.h b/Engine/Plugins/Runtime/AR/Apple/AppleARKit/Source/AppleARKit/Private/AppleARKitSystem.h
--- a/Engine/Plugins/Runtime/AR/Apple/AppleARKit/Source/AppleARKit/Private/AppleARKitSystem.h
+++ b/Engine/Plugins/Runtime/AR/Apple/AppleARKit/Source/AppleARKit/Private/AppleARKitSystem.h
@@ -63,8 +63,11 @@ private:
 	//~ FGCObject
 protected:
 	//~IARSystemSupport
+	/** Returns true/false based on whether AR features are available */
+	virtual bool IsARAvailable() const override;
 	virtual void OnARSystemInitialized() override;
 	virtual EARTrackingQuality OnGetTrackingQuality() const override;
+	virtual EARTrackingQualityReason OnGetTrackingQualityReason() const override;
 	virtual void OnStartARSession(UARSessionConfig* SessionConfig) override;
 	virtual void OnPauseARSession() override;
 	virtual void OnStopARSession() override;
@@ -145,7 +148,7 @@ private:
 	
 	bool bIsRunning = false;
 	
-	void SetDeviceOrientation(EDeviceScreenOrientation InOrientation);
+	void SetDeviceOrientationAndDerivedTracking(EDeviceScreenOrientation InOrientation);
 
 	/** Creates or clears the face ar support object if face ar has been requested */
 	void CheckForFaceARSupport(UARSessionConfig* InSessionConfig);
@@ -213,7 +216,7 @@ private:
 	IAppleARKitPoseTrackingLiveLink* PoseTrackingARLiveLink;
 
 	/** The time code provider to use when tagging time stamps */
-	UTimecodeProvider* TimecodeProvider;
+	UTimecodeProvider* TimecodeProvider = nullptr;
 };
 
 
diff --git a/Engine/Plugins/Runtime/AR/Apple/AppleARKit/Source/AppleARKit/Private/AppleARKitVideoOverlay.cpp b/Engine/Plugins/Runtime/AR/Apple/AppleARKit/Source/AppleARKit/Private/AppleARKitVideoOverlay.cpp
--- a/Engine/Plugins/Runtime/AR/Apple/AppleARKit/Source/AppleARKit/Private/AppleARKitVideoOverlay.cpp
+++ b/Engine/Plugins/Runtime/AR/Apple/AppleARKit/Source/AppleARKit/Private/AppleARKitVideoOverlay.cpp
@@ -28,6 +28,11 @@
 	#include "IOSAppDelegate.h"
 #endif
 
+DECLARE_FLOAT_COUNTER_STAT(TEXT("ARKit Frame to Texture Delay (ms)"), STAT_ARKitFrameToTextureDelay, STATGROUP_ARKIT);
+DECLARE_FLOAT_COUNTER_STAT(TEXT("ARKit Frame to Render Delay (ms)"), STAT_ARKitFrameToRenderDelay, STATGROUP_ARKIT);
+
+
+
 DECLARE_CYCLE_STAT(TEXT("Update Occlusion Textures"), STAT_UpdateOcclusionTextures, STATGROUP_ARKIT);
 
 const FString UARKitCameraOverlayMaterialLoader::OverlayMaterialPath(TEXT("/AppleARKit/M_CameraOverlay.M_CameraOverlay"));
@@ -217,6 +222,12 @@ void FAppleARKitVideoOverlay::RenderVideoOverlayWithMaterial(FRHICommandListImme
 		return;
 	}
 	
+#if STATS && PLATFORM_IOS
+	// LastUpdateTimestamp has been removed
+	//const auto DelayMS = ([[NSProcessInfo processInfo] systemUptime] - LastUpdateTimestamp) * 1000.0;
+	//SET_FLOAT_STAT(STAT_ARKitFrameToRenderDelay, DelayMS);
+#endif
+	
 	SCOPED_DRAW_EVENTF(RHICmdList, RenderVideoOverlay, bRenderingOcclusion ? TEXT("VideoOverlay (Occlusion)") : TEXT("VideoOverlay (Background)"));
 	
 	if (OverlayVertexBufferRHI[0] == nullptr || !OverlayVertexBufferRHI[0].IsValid())
diff --git a/Engine/Plugins/Runtime/AR/Apple/AppleARKit/Source/AppleARKit/Public/AppleARKitCamera.h b/Engine/Plugins/Runtime/AR/Apple/AppleARKit/Source/AppleARKit/Public/AppleARKitCamera.h
--- a/Engine/Plugins/Runtime/AR/Apple/AppleARKit/Source/AppleARKit/Public/AppleARKitCamera.h
+++ b/Engine/Plugins/Runtime/AR/Apple/AppleARKit/Source/AppleARKit/Public/AppleARKitCamera.h
@@ -42,6 +42,7 @@ struct APPLEARKIT_API FAppleARKitCamera
 	// Default constructor
 	FAppleARKitCamera()
 		: TrackingQuality(EARTrackingQuality::NotTracking)
+		, TrackingQualityReason(EARTrackingQualityReason::None)
 		, Orientation(ForceInit)
 		, Translation(ForceInitToZero)
 		, ImageResolution(ForceInitToZero)
@@ -64,6 +65,12 @@ struct APPLEARKIT_API FAppleARKitCamera
 	 */
 	UPROPERTY( BlueprintReadOnly, Category="AppleARKit|Camera" )
 	EARTrackingQuality TrackingQuality;
+	
+	/**
+	 * The reason for the current tracking quality of the camera.
+	 */
+	UPROPERTY( BlueprintReadOnly, Category="AppleARKit|Camera" )
+	EARTrackingQualityReason TrackingQualityReason;
 
 	/**
 	 * The transformation matrix that defines the camera's rotation and translation in world coordinates.
diff --git a/Engine/Plugins/Runtime/AR/Apple/AppleARKit/Source/AppleARKit/Public/AppleARKitFaceSupport.h b/Engine/Plugins/Runtime/AR/Apple/AppleARKit/Source/AppleARKit/Public/AppleARKitFaceSupport.h
--- a/Engine/Plugins/Runtime/AR/Apple/AppleARKit/Source/AppleARKit/Public/AppleARKitFaceSupport.h
+++ b/Engine/Plugins/Runtime/AR/Apple/AppleARKit/Source/AppleARKit/Public/AppleARKitFaceSupport.h
@@ -16,6 +16,7 @@
 struct FAppleARKitAnchorData;
 struct FARVideoFormat;
 class UTimecodeProvider;
+enum class EARFaceTrackingUpdate : uint8;
 
 class APPLEARKIT_API IAppleARKitFaceSupport :
 	public IModularFeature
diff --git a/Engine/Plugins/Runtime/AR/Apple/AppleARKit/Source/AppleARKit/Public/AppleARKitSettings.h b/Engine/Plugins/Runtime/AR/Apple/AppleARKit/Source/AppleARKit/Public/AppleARKitSettings.h
--- a/Engine/Plugins/Runtime/AR/Apple/AppleARKit/Source/AppleARKit/Public/AppleARKitSettings.h
+++ b/Engine/Plugins/Runtime/AR/Apple/AppleARKit/Source/AppleARKit/Public/AppleARKitSettings.h
@@ -8,7 +8,8 @@
 #include "Engine/TimecodeProvider.h"
 
 #include "AppleImageUtilsTypes.h"
-#include "AppleARKitTimecodeProvider.h"
+// MERGE-TODO - this now seems to be private
+#include "../Private/AppleARKitTimecodeProvider.h"
 
 #include "AppleARKitSettings.generated.h"
 
diff --git a/Engine/Plugins/Runtime/AR/Apple/AppleARKitFaceSupport/Source/AppleARKitFaceSupport/Private/AppleARKitFaceSupportImpl.cpp b/Engine/Plugins/Runtime/AR/Apple/AppleARKitFaceSupport/Source/AppleARKitFaceSupport/Private/AppleARKitFaceSupportImpl.cpp
--- a/Engine/Plugins/Runtime/AR/Apple/AppleARKitFaceSupport/Source/AppleARKitFaceSupport/Private/AppleARKitFaceSupportImpl.cpp
+++ b/Engine/Plugins/Runtime/AR/Apple/AppleARKitFaceSupport/Source/AppleARKitFaceSupport/Private/AppleARKitFaceSupportImpl.cpp
@@ -13,6 +13,9 @@
 
 DECLARE_CYCLE_STAT(TEXT("Conversion"), STAT_FaceAR_Conversion, STATGROUP_FaceAR);
 
+// MERGE-todo
+//bool FAppleARKitFaceSupport::bNeedsInit = true;
+
 #if SUPPORTS_ARKIT_1_0
 
 static TSharedPtr<FAppleARKitAnchorData> MakeAnchorData(bool bFaceMirrored, ARAnchor* Anchor, const FRotator& AdjustBy, EARFaceTrackingUpdate UpdateSetting, const FTimecode& Timecode, uint32 FrameRate)
diff --git a/Engine/Plugins/Runtime/AR/Apple/AppleARKitFaceSupport/Source/AppleARKitFaceSupport/Private/AppleARKitFaceSupportImpl.h b/Engine/Plugins/Runtime/AR/Apple/AppleARKitFaceSupport/Source/AppleARKitFaceSupport/Private/AppleARKitFaceSupportImpl.h
--- a/Engine/Plugins/Runtime/AR/Apple/AppleARKitFaceSupport/Source/AppleARKitFaceSupport/Private/AppleARKitFaceSupportImpl.h
+++ b/Engine/Plugins/Runtime/AR/Apple/AppleARKitFaceSupport/Source/AppleARKitFaceSupport/Private/AppleARKitFaceSupportImpl.h
@@ -25,7 +25,7 @@ private:
 	// ~IAppleARKitFaceSupport
 	virtual TArray<TSharedPtr<FAppleARKitAnchorData>> MakeAnchorData(NSArray<ARAnchor*>* NewAnchors, const FRotator& AdjustBy, EARFaceTrackingUpdate UpdateSetting) override;
 	virtual ARConfiguration* ToARConfiguration(UARSessionConfig* InSessionConfig, UTimecodeProvider* InProvider) override;
-	virtual void PublishLiveLinkData(TSharedPtr<FAppleARKitAnchorData> Anchor) override;
+
 	virtual bool DoesSupportFaceAR() override;
 #if SUPPORTS_ARKIT_1_5
 	virtual TArray<FARVideoFormat> ToARConfiguration() override;
@@ -34,6 +34,8 @@ private:
 
 	/** Publishes the remote publisher and the file writer if present */
 	void ProcessRealTimePublishers(TSharedPtr<FAppleARKitAnchorData> AnchorData);
+
+	virtual void PublishLiveLinkData(TSharedPtr<FAppleARKitAnchorData> Anchor) override;
 #endif
 	/** Inits the real time providers if needed */
 	void InitRealtimeProviders();
@@ -44,6 +46,8 @@ private:
 
 	/** Whether the face data is mirrored or not */
 	bool bFaceMirrored;
+	
+protected:
 	/** If requested, publishes face ar updates to LiveLink for the animation system to use */
 	TSharedPtr<ILiveLinkSourceARKit> LiveLinkSource;
 	/** Copied from the UARSessionConfig project settings object */
diff --git a/Engine/Plugins/Runtime/AR/Apple/AppleARKitFaceSupport/Source/AppleARKitFaceSupport/Private/AppleARKitLiveLinkSource.cpp b/Engine/Plugins/Runtime/AR/Apple/AppleARKitFaceSupport/Source/AppleARKitFaceSupport/Private/AppleARKitLiveLinkSource.cpp
--- a/Engine/Plugins/Runtime/AR/Apple/AppleARKitFaceSupport/Source/AppleARKitFaceSupport/Private/AppleARKitLiveLinkSource.cpp
+++ b/Engine/Plugins/Runtime/AR/Apple/AppleARKitFaceSupport/Source/AppleARKitFaceSupport/Private/AppleARKitLiveLinkSource.cpp
@@ -19,6 +19,10 @@
 #include "AppleARKitSettings.h"
 #include "ARTrackable.h"
 
+#include "LiveLinkClient.h"
+
+#include "Animation/PoseAsset.h"
+
 #include "AppleARKitFaceSupportModule.h"
 
 #include "Roles/LiveLinkBasicRole.h"
@@ -359,8 +363,12 @@ FAppleARKitLiveLinkRemoteListener::~FAppleARKitLiveLinkRemoteListener()
 	if (RecvSocket != nullptr)
 	{
 		RecvSocket->Close();
-		ISocketSubsystem* SocketSub = ISocketSubsystem::Get();
-		SocketSub->DestroySocket(RecvSocket);
+
+		if (!GIsRequestingExit)
+		{
+			ISocketSubsystem* SocketSub = ISocketSubsystem::Get();
+			SocketSub->DestroySocket(RecvSocket);
+		}
 	}
 }
 
diff --git a/Engine/Plugins/Runtime/AR/Google/GoogleARCore/Source/GoogleARCoreBase/GoogleARCoreBase.Build.cs b/Engine/Plugins/Runtime/AR/Google/GoogleARCore/Source/GoogleARCoreBase/GoogleARCoreBase.Build.cs
--- a/Engine/Plugins/Runtime/AR/Google/GoogleARCore/Source/GoogleARCoreBase/GoogleARCoreBase.Build.cs
+++ b/Engine/Plugins/Runtime/AR/Google/GoogleARCore/Source/GoogleARCoreBase/GoogleARCoreBase.Build.cs
@@ -62,6 +62,10 @@ namespace UnrealBuildTool.Rules
 
 				// Register Plugin Language
 				AdditionalPropertiesForReceipt.Add("AndroidPlugin", Path.Combine(ModuleDirectory, "GoogleARCoreBase_APL.xml"));
+				
+				// Needed for including "AndroidEGL.h"
+				string EnginePath = Path.GetFullPath(Target.RelativeEnginePath);
+				PrivateIncludePaths.Add(Path.Combine(EnginePath, "Source/Runtime/OpenGLDrv/Private"));
 			}
 
 			if (Target.bBuildEditor)
diff --git a/Engine/Plugins/Runtime/AR/Google/GoogleARCore/Source/GoogleARCoreBase/GoogleARCoreBase_APL.xml b/Engine/Plugins/Runtime/AR/Google/GoogleARCore/Source/GoogleARCoreBase/GoogleARCoreBase_APL.xml
--- a/Engine/Plugins/Runtime/AR/Google/GoogleARCore/Source/GoogleARCoreBase/GoogleARCoreBase_APL.xml
+++ b/Engine/Plugins/Runtime/AR/Google/GoogleARCore/Source/GoogleARCoreBase/GoogleARCoreBase_APL.xml
@@ -113,7 +113,7 @@
   <buildGradleAdditions>
       <insert>
         dependencies {
-          compile('com.google.ar:core:1.9.0')
+          implementation('com.google.ar:core:1.9.0')
         }
       </insert>
   </buildGradleAdditions>
diff --git a/Engine/Plugins/Runtime/AR/Google/GoogleARCore/Source/GoogleARCoreBase/Private/GoogleARCoreAPI.cpp b/Engine/Plugins/Runtime/AR/Google/GoogleARCore/Source/GoogleARCoreBase/Private/GoogleARCoreAPI.cpp
--- a/Engine/Plugins/Runtime/AR/Google/GoogleARCore/Source/GoogleARCoreBase/Private/GoogleARCoreAPI.cpp
+++ b/Engine/Plugins/Runtime/AR/Google/GoogleARCore/Source/GoogleARCoreBase/Private/GoogleARCoreAPI.cpp
@@ -159,10 +159,17 @@ EGoogleARCoreAvailability FGoogleARCoreAPKManager::CheckARCoreAPKAvailability()
 	static JNIEnv* Env = FAndroidApplication::GetJavaEnv();
 	static jmethodID Method = FJavaWrapper::FindMethod(Env, FJavaWrapper::GameActivityClassID, "getApplicationContext", "()Landroid/content/Context;", false);
 	static jobject ApplicationContext = FJavaWrapper::CallObjectMethod(Env, FAndroidApplication::GetGameActivityThis(), Method);
-
-	ArAvailability OutAvailability = AR_AVAILABILITY_UNKNOWN_ERROR;
+	
+	// OutAvailability is cached here as calling 'ArCoreApk_checkAvailability' can cause memory leak in JAVA on some devices
+	// so we bail out immediately if the last call returns OK (this of course assumes the availability doesn't degrades during runtime)
+	static ArAvailability OutAvailability = AR_AVAILABILITY_UNKNOWN_ERROR;
+	if (OutAvailability == AR_AVAILABILITY_SUPPORTED_INSTALLED)
+	{
+		return EGoogleARCoreAvailability::SupportedInstalled;
+	}
+	
 	ArCoreApk_checkAvailability(Env, ApplicationContext, &OutAvailability);
-
+	
 	// Use static_cast here since we already make sure the enum has the same value.
 	return static_cast<EGoogleARCoreAvailability>(OutAvailability);
 #endif
diff --git a/Engine/Plugins/Runtime/AR/Google/GoogleARCore/Source/GoogleARCoreBase/Private/GoogleARCoreCameraImageBlitter.cpp b/Engine/Plugins/Runtime/AR/Google/GoogleARCore/Source/GoogleARCoreBase/Private/GoogleARCoreCameraImageBlitter.cpp
--- a/Engine/Plugins/Runtime/AR/Google/GoogleARCore/Source/GoogleARCoreBase/Private/GoogleARCoreCameraImageBlitter.cpp
+++ b/Engine/Plugins/Runtime/AR/Google/GoogleARCore/Source/GoogleARCoreBase/Private/GoogleARCoreCameraImageBlitter.cpp
@@ -1,6 +1,7 @@
 // Copyright 1998-2019 Epic Games, Inc. All Rights Reserved.
 
 #include "GoogleARCoreCameraImageBlitter.h"
+#include "GoogleARCoreBaseModule.h"
 #include "Engine/Texture2D.h"
 #include "Engine/Texture2DDynamic.h"
 #include "Materials/MaterialInstanceDynamic.h"
@@ -14,6 +15,49 @@
 
 DEFINE_LOG_CATEGORY_STATIC(LogGoogleARCoreImageBlitter, Log, All);
 
+#define NUM_CAMERA_COPIES 5
+
+DECLARE_CYCLE_STAT(TEXT("DoBlit"), STAT_DoBlit, STATGROUP_ARCore);
+DECLARE_CYCLE_STAT(TEXT("  Copy Texture To Vulkan"), STAT_CopyTextureToVulkan, STATGROUP_ARCore);
+DECLARE_CYCLE_STAT(TEXT("    Read Pixels"), STAT_ReadPixels, STATGROUP_ARCore);
+DECLARE_CYCLE_STAT(TEXT("    Update Resource"), STAT_UpdateResource, STATGROUP_ARCore);
+
+
+float GCameraImageScale = 1.f;
+static FAutoConsoleVariableRef CVarCameraImageScale(
+	TEXT("r.ARCore.CameraImageScale"),
+	GCameraImageScale,
+	TEXT("Size scale of the camera image from ARCore.\n")
+	TEXT(" Default: 1.0"),
+	ECVF_Default);
+
+#if PLATFORM_ANDROID
+void CheckGLError()
+{
+	auto Error = glGetError();
+	if (!ensure(Error == GL_NO_ERROR))
+	{
+		UE_LOG(LogGoogleARCoreImageBlitter, Error, TEXT("glGetError: 0x%x"), Error);
+	}
+}
+#endif
+
+static void UpdateOverlayMaterialTexture(UTexture2D* NewTexture)
+{
+	if (auto MaterialInstance = Cast<UMaterialInstanceDynamic>(GetDefault<UGoogleARCoreCameraOverlayMaterialLoader>()->DefaultCameraOverlayMaterial))
+	{
+		if (NewTexture)
+		{
+			static const FName ParamName("CameraTexture");
+			MaterialInstance->SetTextureParameterValue(ParamName, NewTexture);
+		}
+		else
+		{
+			MaterialInstance->ClearParameterValues();
+		}
+	}
+}
+
 FGoogleARCoreDeviceCameraBlitter::FGoogleARCoreDeviceCameraBlitter()
   : CurrentCameraCopy(0)
   , BlitShaderProgram(0)
@@ -26,18 +70,37 @@ FGoogleARCoreDeviceCameraBlitter::FGoogleARCoreDeviceCameraBlitter()
 
 FGoogleARCoreDeviceCameraBlitter::~FGoogleARCoreDeviceCameraBlitter()
 {
-	for(int32 i = 0; i < CameraCopies.Num(); i++)
+#if PLATFORM_ANDROID
+	DeleteOpenGLTextures();
+#endif
+	
+	UpdateOverlayMaterialTexture(nullptr);
+	
+	for (auto Texture : CameraCopies)
 	{
-		delete CameraCopies[i];
+		if (Texture)
+		{
+			Texture->RemoveFromRoot();
+		}
 	}
+	CameraCopies = {};
+	
+	for (auto TextureId : CameraCopyIds)
+	{
+		if (TextureId)
+		{
+			delete TextureId;
+		}
+	}
+	CameraCopyIds = {};
 }
 
-UTexture *FGoogleARCoreDeviceCameraBlitter::GetLastCameraImageTexture()
+UTexture2D* FGoogleARCoreDeviceCameraBlitter::GetLastCameraImageTexture()
 {
-	if(CameraCopies.Num())
+	if (CameraCopies.Num())
 	{
 		int32 LastCameraImage = static_cast<int32>(CurrentCameraCopy) - 1;
-		if(LastCameraImage < 0)
+		if (LastCameraImage < 0)
 		{
 			LastCameraImage = CameraCopies.Num() - 1;
 		}
@@ -49,19 +112,24 @@ UTexture *FGoogleARCoreDeviceCameraBlitter::GetLastCameraImageTexture()
 void FGoogleARCoreDeviceCameraBlitter::LateInit(FIntPoint ImageSize)
 {
 #if PLATFORM_ANDROID
-	FIntPoint CameraSize = ImageSize;
-
+	const auto bUseVulkan = FAndroidMisc::ShouldUseVulkan();
+	
 	// Create a ring buffer of UTexture2Ds to store snapshots of
 	// the camera texture, if we don't already have them.
-	if(!CameraCopies.Num())
+	if (!CameraCopies.Num())
 	{
-		for(uint32 i = 0; i < 5; i++)
+		CameraCopySize.X = (float)ImageSize.X * GCameraImageScale;
+		CameraCopySize.Y = (float)ImageSize.Y * GCameraImageScale;
+		
+		UE_LOG(LogGoogleARCoreImageBlitter, Log, TEXT("Using CameraCopySize (%d x %d), from original image size (%d x %d) and scale factor %.2f"),
+			   CameraCopySize.X, CameraCopySize.Y, ImageSize.X, ImageSize.Y, GCameraImageScale);
+		
+		for (uint32 i = 0; i < NUM_CAMERA_COPIES; i++)
 		{
 			// Make the texture itself.
-			UTexture *CameraCopy = UTexture2D::CreateTransient(CameraSize.X, CameraSize.Y, PF_R8G8B8A8);
+			auto CameraCopy = UTexture2D::CreateTransient(CameraCopySize.X, CameraCopySize.Y, PF_R8G8B8A8);
 			check(CameraCopy);
 			CameraCopy->AddToRoot();
-			FTextureResource *resource = CameraCopy->CreateResource();
 			CameraCopy->Filter = TextureFilter::TF_Nearest;
 			// The camera texture is in SRGB space, so we need to set this flag to true.
 			// Note that in GLES3.1, textures have the SRGB flag set to true will have be converted to 
@@ -74,23 +142,51 @@ void FGoogleARCoreDeviceCameraBlitter::LateInit(FIntPoint ImageSize)
 
 			// Make a place to store the OpenGL texture ID, read
 			// back from the RHI within the render thread.
-			uint32 *TextureId = new uint32;
+			uint32* TextureId = new uint32;
 			*TextureId = 0;
 			CameraCopyIds.Add(TextureId);
-
-			// We have to get the OpenGL texture ID on the render
-			// thread, after the native resource has been
-			// initialized.
-			FTexture2DResource* Resource = ((FTexture2DResource *)resource);
-			ENQUEUE_RENDER_COMMAND(InitCameraBlitter)(
-				[Resource, TextureId, CameraSize](FRHICommandListImmediate& RHICmdList)
+		}
+		
+		if (bUseVulkan)
+		{
+			for (auto TextureId : CameraCopyIds)
+			{
+				glGenTextures(1, TextureId);
+				glBindTexture(GL_TEXTURE_2D, *TextureId);
+				glTexImage2D(GL_TEXTURE_2D, 0, GL_RGBA, CameraCopySize.X, CameraCopySize.Y, 0, GL_RGBA, GL_UNSIGNED_BYTE, 0);
+				glTexParameteri(GL_TEXTURE_2D, GL_TEXTURE_MAG_FILTER, GL_NEAREST);
+				glTexParameteri(GL_TEXTURE_2D, GL_TEXTURE_MIN_FILTER, GL_NEAREST);
+				glBindTexture(GL_TEXTURE_2D, 0);
+			}
+			
+			CheckGLError();
+		}
+	}
+	
+	check(CameraCopies.Num() == CameraCopyIds.Num());
+	
+	if (!bUseVulkan)
+	{
+		// We may not be able to get a valid texture ID back immediately in the next frame
+		// so keep trying until we get a valid one
+		for (auto Index = 0; Index < CameraCopies.Num(); ++Index)
+		{
+			auto Texture = CameraCopies[Index];
+			auto TextureIdPtr = CameraCopyIds[Index];
+			auto Resource = ((FTexture2DResource*)Texture->Resource);
+			if (Texture && TextureIdPtr && !*TextureIdPtr)
+			{
+				ENQUEUE_RENDER_COMMAND(InitCameraBlitter)([Resource, TextureIdPtr](FRHICommandListImmediate& RHICmdList)
 				{
-					Resource->InitRHI();
-					*TextureId = *reinterpret_cast<uint32*>(Resource->GetTexture2DRHI()->GetNativeResource());
+					if (Resource && TextureIdPtr)
+					{
+						*TextureIdPtr = *reinterpret_cast<uint32*>(Resource->GetTexture2DRHI()->GetNativeResource());
+					}
 				});
+			}
 		}
 	}
-
+	
 	// Compile and link the shader program.
 	if(!BlitShaderProgram)
 	{
@@ -197,11 +293,56 @@ void FGoogleARCoreDeviceCameraBlitter::LateInit(FIntPoint ImageSize)
 #endif
 }
 
+#if PLATFORM_ANDROID
+void FGoogleARCoreDeviceCameraBlitter::DeleteOpenGLTextures()
+{
+	if (FAndroidMisc::ShouldUseVulkan())
+	{
+		for (auto TextureId : CameraCopyIds)
+		{
+			if (TextureId && *TextureId)
+			{
+				glDeleteTextures(1, TextureId);
+				*TextureId = 0;
+			}
+		}
+	}
+}
+
+void FGoogleARCoreDeviceCameraBlitter::CopyTextureToVulkan(UTexture2D* TargetTexture)
+{
+	SCOPE_CYCLE_COUNTER(STAT_CopyTextureToVulkan);
+	
+	check(TargetTexture);
+	check(CameraCopySize.X == TargetTexture->GetSizeX());
+	check(CameraCopySize.Y == TargetTexture->GetSizeY());
+	
+	if (TargetTexture->PlatformData)
+	{
+		auto& BulkData = TargetTexture->PlatformData->Mips[0].BulkData;
+		if (auto WriteBuffer = BulkData.Lock(LOCK_READ_WRITE))
+		{
+			{
+				SCOPE_CYCLE_COUNTER(STAT_ReadPixels);
+				glReadPixels(0, 0, CameraCopySize.X, CameraCopySize.Y, GL_RGBA, GL_UNSIGNED_BYTE, WriteBuffer);
+			}
+			
+			BulkData.Unlock();
+			
+			{
+				SCOPE_CYCLE_COUNTER(STAT_UpdateResource);
+				TargetTexture->UpdateResource();
+			}
+		}
+	}
+}
+#endif
+
 void FGoogleARCoreDeviceCameraBlitter::DoBlit(uint32_t CameraTextureId, FIntPoint ImageSize)
 {
 #if PLATFORM_ANDROID
-	FIntPoint CameraSize = ImageSize;
-
+	SCOPE_CYCLE_COUNTER(STAT_DoBlit);
+	
 	LateInit(ImageSize);
 
 	if (CameraCopyIds.Num())
@@ -213,7 +354,7 @@ void FGoogleARCoreDeviceCameraBlitter::DoBlit(uint32_t CameraTextureId, FIntPoin
 			glBindFramebuffer(GL_FRAMEBUFFER, FrameBufferObject);
 			glFramebufferTexture2D(GL_FRAMEBUFFER, GL_COLOR_ATTACHMENT0, GL_TEXTURE_2D, TextureId, 0);
 
-			glViewport(0, 0, CameraSize.X, CameraSize.Y);
+			glViewport(0, 0, CameraCopySize.X, CameraCopySize.Y);
 
 			glDisable(GL_DEPTH_TEST);
 			glDisable(GL_BLEND);
@@ -243,22 +384,23 @@ void FGoogleARCoreDeviceCameraBlitter::DoBlit(uint32_t CameraTextureId, FIntPoin
 			glUniform1i(BlitShaderProgram_Uniform_CameraTexture, 0);
 
 			glDrawArrays(GL_TRIANGLES, 0, 6);
-
+			
+			if (FAndroidMisc::ShouldUseVulkan())
+			{
+				CopyTextureToVulkan(CameraCopies[CurrentCameraCopy]);
+			}
+			
 			glBindTexture(GL_TEXTURE_EXTERNAL_OES, 0);
 			glDisableVertexAttribArray(BlitShaderProgram_Attribute_InPos);
 			glFramebufferTexture2D(GL_FRAMEBUFFER, GL_COLOR_ATTACHMENT0, GL_TEXTURE_2D, 0, 0);
 			glBindFramebuffer(GL_FRAMEBUFFER, 0);
 
 			glFlush();
-
 		}
-
-		// Update the dynamic material instance with the newest
-		// texture.
-		UMaterialInstanceDynamic *DynMat = Cast<UMaterialInstanceDynamic>(
-			GetDefault<UGoogleARCoreCameraOverlayMaterialLoader>()->DefaultCameraOverlayMaterial);
-		DynMat->SetTextureParameterValue(FName("CameraTexture"), CameraCopies[CurrentCameraCopy]);
-
+		
+		// Update the dynamic material instance with the newest texture.
+		UpdateOverlayMaterialTexture(CameraCopies[CurrentCameraCopy]);
+		
 		CurrentCameraCopy++;
 		CurrentCameraCopy %= CameraCopies.Num();
 	}
diff --git a/Engine/Plugins/Runtime/AR/Google/GoogleARCore/Source/GoogleARCoreBase/Private/GoogleARCoreCameraImageBlitter.h b/Engine/Plugins/Runtime/AR/Google/GoogleARCore/Source/GoogleARCoreBase/Private/GoogleARCoreCameraImageBlitter.h
--- a/Engine/Plugins/Runtime/AR/Google/GoogleARCore/Source/GoogleARCoreBase/Private/GoogleARCoreCameraImageBlitter.h
+++ b/Engine/Plugins/Runtime/AR/Google/GoogleARCore/Source/GoogleARCoreBase/Private/GoogleARCoreCameraImageBlitter.h
@@ -15,18 +15,24 @@ public:
 	~FGoogleARCoreDeviceCameraBlitter();
 
 	void DoBlit(uint32_t TextureId, FIntPoint ImageSize);
-	UTexture *GetLastCameraImageTexture();
+	UTexture2D* GetLastCameraImageTexture();
 
 private:
-
 	void LateInit(FIntPoint ImageSize);
-
+	
+#if PLATFORM_ANDROID
+	void CopyTextureToVulkan(UTexture2D* TargetTexture);
+	void DeleteOpenGLTextures();
+#endif
+	
 	uint32 CurrentCameraCopy;
 	uint32 BlitShaderProgram;
 	uint32 BlitShaderProgram_Uniform_CameraTexture;
 	uint32 BlitShaderProgram_Attribute_InPos;
 	uint32 FrameBufferObject;
 	uint32 VertexBufferObject;
-	TArray<UTexture *> CameraCopies;
-    TArray<uint32 *> CameraCopyIds;
+	TArray<UTexture2D*> CameraCopies;
+    TArray<uint32*> CameraCopyIds;
+	
+	FIntPoint CameraCopySize = FIntPoint::ZeroValue;
 };
diff --git a/Engine/Plugins/Runtime/AR/Google/GoogleARCore/Source/GoogleARCoreBase/Private/GoogleARCoreDevice.cpp b/Engine/Plugins/Runtime/AR/Google/GoogleARCore/Source/GoogleARCoreBase/Private/GoogleARCoreDevice.cpp
--- a/Engine/Plugins/Runtime/AR/Google/GoogleARCore/Source/GoogleARCoreBase/Private/GoogleARCoreDevice.cpp
+++ b/Engine/Plugins/Runtime/AR/Google/GoogleARCore/Source/GoogleARCoreBase/Private/GoogleARCoreDevice.cpp
@@ -19,6 +19,15 @@
 
 #include "GoogleARCorePermissionHandler.h"
 
+int32 GARCoreEnableVulkanSupport = 0;
+static FAutoConsoleVariableRef CVarARCoreEnableVulkanSupport(
+	TEXT("r.ARCore.EnableVulkanSupport"),
+	GARCoreEnableVulkanSupport,
+	TEXT("Whether to support Vulkan in the ARCore plugin.\n")
+	TEXT(" 0: Disabled (default)\n")
+	TEXT(" 1: Enabled"),
+	ECVF_Default);
+
 namespace
 {
 	EGoogleARCoreFunctionStatus ToARCoreFunctionStatus(EGoogleARCoreAPIStatus Status)
@@ -44,6 +53,15 @@ namespace
 	}
 }
 
+static bool ShouldUseVulkan()
+{
+#if PLATFORM_ANDROID
+	return FAndroidMisc::ShouldUseVulkan();
+#else
+	return false;
+#endif
+}
+
 FGoogleARCoreDelegates::FGoogleARCoreOnConfigCameraDelegate FGoogleARCoreDelegates::OnCameraConfig;
 
 FGoogleARCoreDevice* FGoogleARCoreDevice::GetInstance()
@@ -130,7 +148,15 @@ float FGoogleARCoreDevice::GetWorldToMetersScale()
 void FGoogleARCoreDevice::StartARCoreSessionRequest(UARSessionConfig* SessionConfig)
 {
 	UE_LOG(LogGoogleARCore, Log, TEXT("Start ARCore session requested."));
-
+	
+#if PLATFORM_ANDROID
+	if (ShouldUseVulkan() && !GARCoreEnableVulkanSupport)
+	{
+		UE_LOG(LogGoogleARCore, Error, TEXT("StartARCoreSessionRequest ignored - ARCore requires an OpenGL context, but we're using Vulkan!"));
+		return;
+	}
+#endif
+	
 	// The new SessionConfig should be set to the session already at this point. We can check if it is front camera session here.
 	bool bShouldUseFrontCamera = GetIsFrontCameraSession();
 	bool bShouldSwitchCamera = bShouldUseFrontCamera ? ARCoreSession == BackCameraARCoreSession : ARCoreSession == FrontCameraARCoreSession;
@@ -340,7 +366,12 @@ void FGoogleARCoreDevice::OnWorldTickStart(ELevelTick TickType, float DeltaTime)
 		else
 		{
 			// Copy the camera GPU texture to a normal Texture2D so that we can get around the multi-threaded rendering issue.
-			CameraBlitter.DoBlit(PassthroughCameraTextureId, SessionCameraConfig.CameraTextureResolution);
+			if (!CameraBlitter)
+			{
+				CameraBlitter = MakeShared<FGoogleARCoreDeviceCameraBlitter, ESPMode::ThreadSafe>();
+			}
+			
+			CameraBlitter->DoBlit(PassthroughCameraTextureId, SessionCameraConfig.CameraTextureResolution);
 		}
 	}
 }
@@ -407,7 +438,12 @@ void FGoogleARCoreDevice::HandleRuntimePermissionsGranted(const TArray<FString>&
 void FGoogleARCoreDevice::StartSessionWithRequestedConfig()
 {
 	bStartSessionRequested = false;
-
+	
+	if (ShouldUseVulkan())
+	{
+		GLContext = FGoogleARCoreOpenGLContext::CreateContext();
+	}
+	
 	// Allocate passthrough camera texture if necessary.
 	if (PassthroughCameraTexture == nullptr)
 	{
@@ -628,18 +664,32 @@ void FGoogleARCoreDevice::ResetARCoreSession()
 	BackCameraARCoreSession.Reset();
 	CurrentSessionStatus.Status = EARSessionStatus::NotStarted;
 	CurrentSessionStatus.AdditionalInfo = TEXT("ARCore Session is uninitialized.");
+	// Release the CameraBlitter first as it may need to release resources created from the GLContext
+	if (GLContext)
+	{
+		CameraBlitter = nullptr;
+		GLContext = nullptr;
+	}
 }
 
 void FGoogleARCoreDevice::AllocatePassthroughCameraTexture_RenderThread()
 {
-	FRHICommandListImmediate& RHICmdList = FRHICommandListExecutor::GetImmediateCommandList();
-	FRHIResourceCreateInfo CreateInfo;
-
-	PassthroughCameraTexture = RHICmdList.CreateTextureExternal2D(1, 1, PF_R8G8B8A8, 1, 1, 0, CreateInfo);
-
-	void* NativeResource = PassthroughCameraTexture->GetNativeResource();
-	check(NativeResource);
-	PassthroughCameraTextureId = *reinterpret_cast<uint32*>(NativeResource);
+	if (ShouldUseVulkan())
+	{
+		check(GLContext);
+		PassthroughCameraTextureId = GLContext->GetCameraTextureId();
+	}
+	else
+	{
+		FRHICommandListImmediate& RHICmdList = FRHICommandListExecutor::GetImmediateCommandList();
+		FRHIResourceCreateInfo CreateInfo;
+		
+		PassthroughCameraTexture = RHICmdList.CreateTextureExternal2D(1, 1, PF_R8G8B8A8, 1, 1, 0, CreateInfo);
+		
+		void* NativeResource = PassthroughCameraTexture->GetNativeResource();
+		check(NativeResource);
+		PassthroughCameraTextureId = *reinterpret_cast<uint32*>(NativeResource);
+	}
 }
 
 FTextureRHIRef FGoogleARCoreDevice::GetPassthroughCameraTexture()
@@ -741,7 +791,12 @@ EGoogleARCoreFunctionStatus FGoogleARCoreDevice::GetLatestCameraMetadata(const A
 
 UTexture* FGoogleARCoreDevice::GetCameraTexture()
 {
-	return CameraBlitter.GetLastCameraImageTexture();
+	if (CameraBlitter.IsValid())
+	{
+		return CameraBlitter->GetLastCameraImageTexture();
+	}
+
+	return nullptr;
 }
 
 EGoogleARCoreFunctionStatus FGoogleARCoreDevice::AcquireCameraImage(UGoogleARCoreCameraImage *&OutLatestCameraImage)
diff --git a/Engine/Plugins/Runtime/AR/Google/GoogleARCore/Source/GoogleARCoreBase/Private/GoogleARCoreDevice.h b/Engine/Plugins/Runtime/AR/Google/GoogleARCore/Source/GoogleARCoreBase/Private/GoogleARCoreDevice.h
--- a/Engine/Plugins/Runtime/AR/Google/GoogleARCore/Source/GoogleARCoreBase/Private/GoogleARCoreDevice.h
+++ b/Engine/Plugins/Runtime/AR/Google/GoogleARCore/Source/GoogleARCoreBase/Private/GoogleARCoreDevice.h
@@ -12,6 +12,8 @@
 #include "GoogleARCoreAPI.h"
 #include "GoogleARCorePassthroughCameraRenderer.h"
 #include "GoogleARCoreAugmentedImageDatabase.h"
+#include "GoogleARCoreOpenGLContext.h"
+
 
 class FGoogleARCoreDevice
 {
@@ -189,9 +191,12 @@ private:
 	FARSessionStatus CurrentSessionStatus;
 
 	FGoogleARCoreCameraConfig SessionCameraConfig;
-	FGoogleARCoreDeviceCameraBlitter CameraBlitter;
+	
+	TSharedPtr<FGoogleARCoreDeviceCameraBlitter, ESPMode::ThreadSafe> CameraBlitter;
 
 	TQueue<TFunction<void()>> RunOnGameThreadQueue;
 
 	TSharedPtr<FARSupportInterface, ESPMode::ThreadSafe> ARSystem;
+	
+	TSharedPtr<FGoogleARCoreOpenGLContext, ESPMode::ThreadSafe> GLContext;
 };
diff --git a/Engine/Plugins/Runtime/AR/Google/GoogleARCore/Source/GoogleARCoreBase/Private/GoogleARCoreOpenGLContext.cpp b/Engine/Plugins/Runtime/AR/Google/GoogleARCore/Source/GoogleARCoreBase/Private/GoogleARCoreOpenGLContext.cpp
new file mode 100644
--- /dev/null
+++ b/Engine/Plugins/Runtime/AR/Google/GoogleARCore/Source/GoogleARCoreBase/Private/GoogleARCoreOpenGLContext.cpp
@@ -0,0 +1,91 @@
+// Copyright 2017 Google Inc.
+
+#include "GoogleARCoreOpenGLContext.h"
+
+
+DEFINE_LOG_CATEGORY_STATIC(LogARCoreOpenGL, Log, All);
+
+TSharedPtr<FGoogleARCoreOpenGLContext, ESPMode::ThreadSafe> FGoogleARCoreOpenGLContext::CreateContext()
+{
+#if PLATFORM_ANDROID
+	if (FAndroidMisc::ShouldUseVulkan())
+	{
+		auto Context = MakeShared<FGoogleARCoreOpenGLContext, ESPMode::ThreadSafe>();
+		Context->InitContext();
+		return Context;
+	}
+#endif
+	
+	return nullptr;
+}
+
+FGoogleARCoreOpenGLContext::~FGoogleARCoreOpenGLContext()
+{
+	ReleaseContext();
+	UE_LOG(LogARCoreOpenGL, Log, TEXT("context destroyed"));
+}
+
+void FGoogleARCoreOpenGLContext::InitContext()
+{
+#if PLATFORM_ANDROID
+	UE_LOG(LogARCoreOpenGL, Log, TEXT("FGoogleARCoreOpenGLContext::InitContext"));
+	
+	auto EGL = AndroidEGL::GetInstance();
+	EGL->Init(AndroidEGL::AV_OpenGLES, 2, 0, false);
+	EGL->InitSurface(false, false);
+	SaveContext();
+	MakeCurrent();
+	EGL->InitBackBuffer(); //can be done only after context is made current.
+	glGenTextures(1, &PassthroughCameraTextureId);
+	RestoreContext();
+	
+	UE_LOG(LogARCoreOpenGL, Log, TEXT("PassthroughCameraTextureId: %d"), PassthroughCameraTextureId);
+#endif
+}
+
+void FGoogleARCoreOpenGLContext::ReleaseContext()
+{
+#if PLATFORM_ANDROID
+	glDeleteTextures(1, &PassthroughCameraTextureId);
+	PassthroughCameraTextureId = 0;
+	AndroidEGL::GetInstance()->Terminate();
+#endif
+}
+
+void FGoogleARCoreOpenGLContext::MakeCurrent()
+{
+#if PLATFORM_ANDROID
+	AndroidEGL::GetInstance()->SetCurrentRenderingContext();
+#endif
+}
+
+void FGoogleARCoreOpenGLContext::SaveContext()
+{
+#if PLATFORM_ANDROID
+	SavedDisplay = eglGetCurrentDisplay();
+	SavedContext = eglGetCurrentContext();
+	SavedDrawSurface = eglGetCurrentSurface(EGL_DRAW);
+	SavedReadSurface = eglGetCurrentSurface(EGL_READ);
+#endif
+}
+
+void FGoogleARCoreOpenGLContext::RestoreContext()
+{
+#if PLATFORM_ANDROID
+	eglMakeCurrent(SavedDisplay, SavedDrawSurface, SavedReadSurface, SavedContext);
+	
+	SavedDisplay = EGL_NO_DISPLAY;
+	SavedContext = EGL_NO_CONTEXT;
+	SavedDrawSurface = EGL_NO_SURFACE;
+	SavedReadSurface = EGL_NO_SURFACE;
+#endif
+}
+
+uint32 FGoogleARCoreOpenGLContext::GetCameraTextureId() const
+{
+#if PLATFORM_ANDROID
+	return PassthroughCameraTextureId;
+#else
+	return 0;
+#endif
+}
diff --git a/Engine/Plugins/Runtime/AR/Google/GoogleARCore/Source/GoogleARCoreBase/Private/GoogleARCoreOpenGLContext.h b/Engine/Plugins/Runtime/AR/Google/GoogleARCore/Source/GoogleARCoreBase/Private/GoogleARCoreOpenGLContext.h
new file mode 100644
--- /dev/null
+++ b/Engine/Plugins/Runtime/AR/Google/GoogleARCore/Source/GoogleARCoreBase/Private/GoogleARCoreOpenGLContext.h
@@ -0,0 +1,36 @@
+// Copyright 2017 Google Inc.
+
+#pragma once
+
+#include "Runtime/Core/Public/Misc/Build.h"
+#include "HAL/Platform.h"
+#include "Templates/SharedPointer.h"
+
+#if PLATFORM_ANDROID
+#include "Android/AndroidEGL.h"
+#endif
+
+class FGoogleARCoreOpenGLContext
+{
+public:
+	static TSharedPtr<FGoogleARCoreOpenGLContext, ESPMode::ThreadSafe> CreateContext();
+	~FGoogleARCoreOpenGLContext();
+	
+	void InitContext();
+	void ReleaseContext();
+	
+	void SaveContext();
+	void MakeCurrent();
+	void RestoreContext();
+	
+	uint32 GetCameraTextureId() const;
+	
+private:
+#if PLATFORM_ANDROID
+	EGLDisplay SavedDisplay = EGL_NO_DISPLAY;
+	EGLContext SavedContext = EGL_NO_CONTEXT;
+	EGLSurface SavedDrawSurface = EGL_NO_SURFACE;
+	EGLSurface SavedReadSurface = EGL_NO_SURFACE;
+	uint32 PassthroughCameraTextureId = 0;
+#endif
+};
diff --git a/Engine/Plugins/Runtime/AR/Google/GoogleARCore/Source/GoogleARCoreBase/Private/GoogleARCoreXRTrackingSystem.cpp b/Engine/Plugins/Runtime/AR/Google/GoogleARCore/Source/GoogleARCoreBase/Private/GoogleARCoreXRTrackingSystem.cpp
--- a/Engine/Plugins/Runtime/AR/Google/GoogleARCore/Source/GoogleARCoreBase/Private/GoogleARCoreXRTrackingSystem.cpp
+++ b/Engine/Plugins/Runtime/AR/Google/GoogleARCore/Source/GoogleARCoreBase/Private/GoogleARCoreXRTrackingSystem.cpp
@@ -195,7 +195,7 @@ EARTrackingQuality FGoogleARCoreXRTrackingSystem::OnGetTrackingQuality() const
 	return EARTrackingQuality::OrientationAndPosition;
 }
 
-/*
+
 EARTrackingQualityReason FGoogleARCoreXRTrackingSystem::OnGetTrackingQualityReason() const
 {
 	// Dont return EARTrackingQualityReason::None, which means that the tracking quality is not limited
@@ -206,7 +206,7 @@ bool FGoogleARCoreXRTrackingSystem::IsARAvailable() const
 {
 	return FGoogleARCoreDevice::GetInstance()->CheckARCoreAPKAvailability() == EGoogleARCoreAvailability::SupportedInstalled;
 }
-*/
+
 
 void FGoogleARCoreXRTrackingSystem::OnStartARSession(UARSessionConfig* SessionConfig)
 {
diff --git a/Engine/Plugins/Runtime/AR/Google/GoogleARCore/Source/GoogleARCoreBase/Private/GoogleARCoreXRTrackingSystem.h b/Engine/Plugins/Runtime/AR/Google/GoogleARCore/Source/GoogleARCoreBase/Private/GoogleARCoreXRTrackingSystem.h
--- a/Engine/Plugins/Runtime/AR/Google/GoogleARCore/Source/GoogleARCoreBase/Private/GoogleARCoreXRTrackingSystem.h
+++ b/Engine/Plugins/Runtime/AR/Google/GoogleARCore/Source/GoogleARCoreBase/Private/GoogleARCoreXRTrackingSystem.h
@@ -48,12 +48,10 @@ public:
 
 protected:
 	// IARSystemSupport
-	/** todo - reenable once base changes make their way up - Returns true/false based on whether AR features are available */
-	//virtual bool IsARAvailable() const override;
+	virtual bool IsARAvailable() const override;
 	virtual void OnARSystemInitialized() override;
 	virtual EARTrackingQuality OnGetTrackingQuality() const override;
-	/** todo - reenable once base changes make their way up */
-	//virtual EARTrackingQualityReason OnGetTrackingQualityReason() const override;
+	virtual EARTrackingQualityReason OnGetTrackingQualityReason() const override;
 	virtual void OnStartARSession(UARSessionConfig* SessionConfig) override;
 	virtual void OnPauseARSession() override;
 	virtual void OnStopARSession() override;
diff --git a/Engine/Plugins/Runtime/AR/Google/GoogleARCore/Source/GoogleARCoreBase/Public/GoogleARCoreBaseModule.h b/Engine/Plugins/Runtime/AR/Google/GoogleARCore/Source/GoogleARCoreBase/Public/GoogleARCoreBaseModule.h
--- a/Engine/Plugins/Runtime/AR/Google/GoogleARCore/Source/GoogleARCoreBase/Public/GoogleARCoreBaseModule.h
+++ b/Engine/Plugins/Runtime/AR/Google/GoogleARCore/Source/GoogleARCoreBase/Public/GoogleARCoreBaseModule.h
@@ -3,6 +3,7 @@
 #pragma once
 
 #include "Modules/ModuleManager.h"
+#include "Stats/Stats.h"
 #include "IHeadMountedDisplayModule.h"
 
 /**
@@ -35,3 +36,5 @@ public:
 		return FModuleManager::Get().IsModuleLoaded("GoogleARCoreBase");
 	}
 };
+
+DECLARE_STATS_GROUP(TEXT("ARCore"), STATGROUP_ARCore, STATCAT_Advanced);
diff --git a/Engine/Plugins/Runtime/AR/Google/GoogleARCore/Source/GoogleARCoreRendering/Private/GoogleARCorePassthroughCameraRenderer.cpp b/Engine/Plugins/Runtime/AR/Google/GoogleARCore/Source/GoogleARCoreRendering/Private/GoogleARCorePassthroughCameraRenderer.cpp
--- a/Engine/Plugins/Runtime/AR/Google/GoogleARCore/Source/GoogleARCoreRendering/Private/GoogleARCorePassthroughCameraRenderer.cpp
+++ b/Engine/Plugins/Runtime/AR/Google/GoogleARCore/Source/GoogleARCoreRendering/Private/GoogleARCorePassthroughCameraRenderer.cpp
@@ -92,7 +92,11 @@ void FGoogleARCorePassthroughCameraRenderer::InitializeRenderer_RenderThread(FTe
 	FSamplerStateInitializerRHI SamplerStateInitializer(SF_Point, AM_Clamp, AM_Clamp, AM_Clamp);
 	FSamplerStateRHIRef SamplerStateRHI = RHICreateSamplerState(SamplerStateInitializer);
 
-	FExternalTextureRegistry::Get().RegisterExternalTexture(GoogleARCorePassthroughCameraExternalTextureGuid, VideoTexture, SamplerStateRHI);
+	// VideoTexture can be NULL for Vulkan
+	if (VideoTexture)
+	{
+		FExternalTextureRegistry::Get().RegisterExternalTexture(GoogleARCorePassthroughCameraExternalTextureGuid, VideoTexture, SamplerStateRHI);
+	}
 
 	//Make sure AR camera pass through materials are updated properly
 	FMaterialRenderProxy::UpdateDeferredCachedUniformExpressions();
@@ -217,6 +221,13 @@ IMPLEMENT_MATERIAL_SHADER(FGoogleARCoreCameraOverlayPS, "/Engine/Private/PostPro
 void FGoogleARCorePassthroughCameraRenderer::RenderVideoOverlay_RenderThread(FRHICommandListImmediate& RHICmdList, FSceneView& InView)
 {
 #if PLATFORM_ANDROID
+	if (FAndroidMisc::ShouldUseVulkan() && IsMobileHDR() && !RHICmdList.IsInsideRenderPass())
+	{
+		// We must NOT call DrawIndexedPrimitive below if not in a render pass on Vulkan, it's very likely to crash!
+		UE_LOG(LogTemp, Warning, TEXT("FGoogleARCorePassthroughCameraRenderer::RenderVideoOverlay_RenderThread: skipped due to not called within a render pass on Vulkan!"));
+		return;
+	}
+	
 	if (RenderingOverlayMaterial == nullptr || !RenderingOverlayMaterial->IsValidLowLevel())
 	{
 		return;
@@ -244,7 +255,7 @@ void FGoogleARCorePassthroughCameraRenderer::RenderVideoOverlay_RenderThread(FRH
 		GraphicsPSOInit.BoundShaderState.PixelShaderRHI = GETSAFERHISHADER_PIXEL(PixelShader);
 		GraphicsPSOInit.PrimitiveType = PT_TriangleList;
 
-		SetGraphicsPipelineState(RHICmdList, GraphicsPSOInit, EApplyRendertargetOption::DoNothing);
+		SetGraphicsPipelineState(RHICmdList, GraphicsPSOInit);
 
 		VertexShader->SetParameters(RHICmdList, InView);
 		PixelShader->SetParameters(RHICmdList, InView, RenderingOverlayMaterial->GetRenderProxy());
diff --git a/Engine/Source/Runtime/AugmentedReality/Public/ARSystem.h b/Engine/Source/Runtime/AugmentedReality/Public/ARSystem.h
--- a/Engine/Source/Runtime/AugmentedReality/Public/ARSystem.h
+++ b/Engine/Source/Runtime/AugmentedReality/Public/ARSystem.h
@@ -84,8 +84,7 @@ public:
 	virtual EARTrackingQuality OnGetTrackingQuality() const = 0;
 	
 	/** @return the reason of limited tracking quality; if the state is not limited, return EARTrackingQualityReason::None */
-	// @todo merge-check
-//	virtual EARTrackingQualityReason OnGetTrackingQualityReason() const = 0;
+	virtual EARTrackingQualityReason OnGetTrackingQualityReason() const = 0;
 	
 	/**
 	 * Start the AR system.
@@ -106,7 +105,7 @@ public:
 	virtual FARSessionStatus OnGetARSessionStatus() const = 0;
 	
 	/** Returns true/false based on whether AR features are available */
-//	virtual bool IsARAvailable() const = 0;
+	virtual bool IsARAvailable() const = 0;
 	
 	/**
 	 * Set a transform that will align the Tracking Space origin to the World Space origin.
