Add which UWorld that FWorldDelegates::OnWorldTickStart is being called for

diff --git a/Engine/Plugins/Lumin/MagicLeap/Source/MagicLeapAR/Private/LuminARDevice.cpp b/Engine/Plugins/Lumin/MagicLeap/Source/MagicLeapAR/Private/LuminARDevice.cpp
--- a/Engine/Plugins/Lumin/MagicLeap/Source/MagicLeapAR/Private/LuminARDevice.cpp
+++ b/Engine/Plugins/Lumin/MagicLeap/Source/MagicLeapAR/Private/LuminARDevice.cpp
@@ -121,9 +121,9 @@ bool FLuminARDevice::GetStartSessionRequestFinished()
 }
 
 // Note that this function will only be registered when LuminAR is supported.
-void FLuminARDevice::OnWorldTickStart(ELevelTick TickType, float DeltaTime)
+void FLuminARDevice::OnWorldTickStart(UWorld* World, ELevelTick TickType, float DeltaTime)
 {
-	WorldToMeterScale = GWorld->GetWorldSettings()->WorldToMeters;
+	WorldToMeterScale = World->GetWorldSettings()->WorldToMeters;
 	TFunction<void()> Func;
 	while (RunOnGameThreadQueue.Dequeue(Func))
 	{
diff --git a/Engine/Plugins/Lumin/MagicLeap/Source/MagicLeapAR/Private/LuminARDevice.h b/Engine/Plugins/Lumin/MagicLeap/Source/MagicLeapAR/Private/LuminARDevice.h
--- a/Engine/Plugins/Lumin/MagicLeap/Source/MagicLeapAR/Private/LuminARDevice.h
+++ b/Engine/Plugins/Lumin/MagicLeap/Source/MagicLeapAR/Private/LuminARDevice.h
@@ -98,7 +98,7 @@ private:
 	void OnModuleLoaded();
 	void OnModuleUnloaded();
 
-	void OnWorldTickStart(ELevelTick TickType, float DeltaTime);
+	void OnWorldTickStart(UWorld* World, ELevelTick TickType, float DeltaTime);
 
 	void StartSession();
 
diff --git a/Engine/Plugins/Runtime/AR/Google/GoogleARCore/Source/GoogleARCoreBase/Private/GoogleARCoreDevice.cpp b/Engine/Plugins/Runtime/AR/Google/GoogleARCore/Source/GoogleARCoreBase/Private/GoogleARCoreDevice.cpp
--- a/Engine/Plugins/Runtime/AR/Google/GoogleARCore/Source/GoogleARCoreBase/Private/GoogleARCoreDevice.cpp
+++ b/Engine/Plugins/Runtime/AR/Google/GoogleARCore/Source/GoogleARCoreBase/Private/GoogleARCoreDevice.cpp
@@ -270,9 +270,9 @@ bool FGoogleARCoreDevice::GetStartSessionRequestFinished()
 }
 
 // Note that this function will only be registered when ARCore is supported.
-void FGoogleARCoreDevice::OnWorldTickStart(ELevelTick TickType, float DeltaTime)
+void FGoogleARCoreDevice::OnWorldTickStart(UWorld* World, ELevelTick TickType, float DeltaTime)
 {
-	WorldToMeterScale = GWorld->GetWorldSettings()->WorldToMeters;
+	WorldToMeterScale = World->GetWorldSettings()->WorldToMeters;
 	TFunction<void()> Func;
 	while (RunOnGameThreadQueue.Dequeue(Func))
 	{
diff --git a/Engine/Plugins/Runtime/AR/Google/GoogleARCore/Source/GoogleARCoreBase/Private/GoogleARCoreDevice.h b/Engine/Plugins/Runtime/AR/Google/GoogleARCore/Source/GoogleARCoreBase/Private/GoogleARCoreDevice.h
--- a/Engine/Plugins/Runtime/AR/Google/GoogleARCore/Source/GoogleARCoreBase/Private/GoogleARCoreDevice.h
+++ b/Engine/Plugins/Runtime/AR/Google/GoogleARCore/Source/GoogleARCoreBase/Private/GoogleARCoreDevice.h
@@ -153,7 +153,7 @@ private:
 	void OnModuleLoaded();
 	void OnModuleUnloaded();
 
-	void OnWorldTickStart(ELevelTick TickType, float DeltaTime);
+	void OnWorldTickStart(UWorld* World, ELevelTick TickType, float DeltaTime);
 
 	void CheckAndRequrestPermission(const UARSessionConfig& ConfigurationData);
 
diff --git a/Engine/Plugins/Runtime/AR/Google/GoogleARCoreServices/Source/GoogleARCoreServices/Private/GoogleARCoreServicesManager.cpp b/Engine/Plugins/Runtime/AR/Google/GoogleARCoreServices/Source/GoogleARCoreServices/Private/GoogleARCoreServicesManager.cpp
--- a/Engine/Plugins/Runtime/AR/Google/GoogleARCoreServices/Source/GoogleARCoreServices/Private/GoogleARCoreServicesManager.cpp
+++ b/Engine/Plugins/Runtime/AR/Google/GoogleARCoreServices/Source/GoogleARCoreServices/Private/GoogleARCoreServicesManager.cpp
@@ -140,7 +140,7 @@ void FGoogleARCoreServicesManager::OnARSessionStarted()
 }
 
 // We need to make sure this tick happens after the ARSystem tick
-void FGoogleARCoreServicesManager::OnWorldTickStart(ELevelTick TickType, float DeltaTime)
+void FGoogleARCoreServicesManager::OnWorldTickStart(UWorld* World, ELevelTick TickType, float DeltaTime)
 {
 	if (!bHasValidARSystem)
 	{
diff --git a/Engine/Plugins/Runtime/AR/Google/GoogleARCoreServices/Source/GoogleARCoreServices/Private/GoogleARCoreServicesManager.h b/Engine/Plugins/Runtime/AR/Google/GoogleARCoreServices/Source/GoogleARCoreServices/Private/GoogleARCoreServicesManager.h
--- a/Engine/Plugins/Runtime/AR/Google/GoogleARCoreServices/Source/GoogleARCoreServices/Private/GoogleARCoreServicesManager.h
+++ b/Engine/Plugins/Runtime/AR/Google/GoogleARCoreServices/Source/GoogleARCoreServices/Private/GoogleARCoreServicesManager.h
@@ -32,7 +32,7 @@ private:
 	EARPinCloudTaskResult CheckCloudTaskError();
 
 	void OnARSessionStarted();
-	void OnWorldTickStart(ELevelTick TickType, float DeltaTime);
+	void OnWorldTickStart(UWorld* World, ELevelTick TickType, float DeltaTime);
 
 	bool bHasValidARSystem;
 	bool bCloudARPinEnabled;
diff --git a/Engine/Plugins/Runtime/AR/Microsoft/HoloLensAR/Source/HoloLensAR/Private/HoloLensARSystem.cpp b/Engine/Plugins/Runtime/AR/Microsoft/HoloLensAR/Source/HoloLensAR/Private/HoloLensARSystem.cpp
--- a/Engine/Plugins/Runtime/AR/Microsoft/HoloLensAR/Source/HoloLensAR/Private/HoloLensARSystem.cpp
+++ b/Engine/Plugins/Runtime/AR/Microsoft/HoloLensAR/Source/HoloLensAR/Private/HoloLensARSystem.cpp
@@ -90,7 +90,7 @@ void FHoloLensARSystem::Shutdown()
 	SessionConfig = nullptr;
 }
 
-void FHoloLensARSystem::OnWorldTickStart(ELevelTick TickType, float DeltaTime)
+void FHoloLensARSystem::OnWorldTickStart(UWorld* World, ELevelTick TickType, float DeltaTime)
 {
 	UpdateWMRAnchors();
 }
diff --git a/Engine/Plugins/Runtime/AR/Microsoft/HoloLensAR/Source/HoloLensAR/Public/HoloLensARSystem.h b/Engine/Plugins/Runtime/AR/Microsoft/HoloLensAR/Source/HoloLensAR/Public/HoloLensARSystem.h
--- a/Engine/Plugins/Runtime/AR/Microsoft/HoloLensAR/Source/HoloLensAR/Public/HoloLensARSystem.h
+++ b/Engine/Plugins/Runtime/AR/Microsoft/HoloLensAR/Source/HoloLensAR/Public/HoloLensARSystem.h
@@ -68,7 +68,7 @@ public:
 	void Shutdown();
 
 private:
-	void OnWorldTickStart(ELevelTick TickType, float DeltaTime);
+	void OnWorldTickStart(UWorld* World, ELevelTick TickType, float DeltaTime);
 
 	//~ FGCObject
 	virtual void AddReferencedObjects(FReferenceCollector& Collector) override;
diff --git a/Engine/Source/Runtime/Engine/Classes/Engine/World.h b/Engine/Source/Runtime/Engine/Classes/Engine/World.h
--- a/Engine/Source/Runtime/Engine/Classes/Engine/World.h
+++ b/Engine/Source/Runtime/Engine/Classes/Engine/World.h
@@ -3562,7 +3562,7 @@ public:
 	// delegate for generating world asset registry tags so project/game scope can add additional tags for filtering levels in their UI, etc
 	DECLARE_MULTICAST_DELEGATE_TwoParams(FWorldGetAssetTags, const UWorld*, TArray<UObject::FAssetRegistryTag>&);
 
-	DECLARE_MULTICAST_DELEGATE_TwoParams(FOnWorldTickStart, ELevelTick, float);
+	DECLARE_MULTICAST_DELEGATE_ThreeParams(FOnWorldTickStart, UWorld*, ELevelTick, float);
 	static FOnWorldTickStart OnWorldTickStart;
 
 	// Delegate called before actors are ticked for each world. Delta seconds is already dilated and clamped.
diff --git a/Engine/Source/Runtime/Engine/Private/LevelTick.cpp b/Engine/Source/Runtime/Engine/Private/LevelTick.cpp
--- a/Engine/Source/Runtime/Engine/Private/LevelTick.cpp
+++ b/Engine/Source/Runtime/Engine/Private/LevelTick.cpp
@@ -1351,7 +1351,7 @@ void UWorld::Tick( ELevelTick TickType, float DeltaSeconds )
 
 	TDrawEvent<FRHICommandList>* TickDrawEvent = BeginTickDrawEvent();
 
-	FWorldDelegates::OnWorldTickStart.Broadcast(TickType, DeltaSeconds);
+	FWorldDelegates::OnWorldTickStart.Broadcast(this, TickType, DeltaSeconds);
 
 	//Tick game and other thread trackers.
 	for (int32 Tracker = 0; Tracker < (int32)EInGamePerfTrackers::Num; ++Tracker)
diff --git a/Engine/Source/Runtime/RuntimeAssetCache/Private/RuntimeAssetCacheModule.cpp b/Engine/Source/Runtime/RuntimeAssetCache/Private/RuntimeAssetCacheModule.cpp
--- a/Engine/Source/Runtime/RuntimeAssetCache/Private/RuntimeAssetCacheModule.cpp
+++ b/Engine/Source/Runtime/RuntimeAssetCache/Private/RuntimeAssetCacheModule.cpp
@@ -37,7 +37,7 @@ public:
 		FWorldDelegates::OnWorldTickStart.AddStatic(&FRuntimeAssetCacheModule::TickRuntimeAssetCache);
 	}
 
-	static void TickRuntimeAssetCache(ELevelTick TickType, float DeltaSeconds)
+	static void TickRuntimeAssetCache(UWorld* World, ELevelTick TickType, float DeltaSeconds)
 	{
 		static FRuntimeAssetCacheModuleInterface& Module = FModuleManager::LoadModuleChecked<FRuntimeAssetCacheModuleInterface>("RuntimeAssetCache");
 		static FRuntimeAssetCacheInterface& Interface = Module.GetRuntimeAssetCache();
