Grouped Apple/Google AR plugins under an AR directory in Experimental to improve the organization of various AR plugins

diff --git a/Engine/Plugins/Experimental/AR/Google/GoogleARCore/Source/GoogleARCoreBase/Private/GoogleARCoreCameraImageBlitter.cpp b/Engine/Plugins/Experimental/AR/Google/GoogleARCore/Source/GoogleARCoreBase/Private/GoogleARCoreCameraImageBlitter.cpp
new file mode 100644
--- /dev/null
+++ b/Engine/Plugins/Experimental/AR/Google/GoogleARCore/Source/GoogleARCoreBase/Private/GoogleARCoreCameraImageBlitter.cpp
@@ -0,0 +1,259 @@
+// Copyright 2017 Google Inc.
+
+#include "GoogleARCoreCameraImageBlitter.h"
+#include "Engine/Texture2D.h"
+#include "Engine/Texture2DDynamic.h"
+#include "Materials/MaterialInstanceDynamic.h"
+#include "GoogleARCorePassthroughCameraRenderer.h"
+
+#if PLATFORM_ANDROID
+#include "Android/AndroidApplication.h"
+#include <GLES2/gl2.h>
+#include <GLES2/gl2ext.h>
+#endif
+
+DEFINE_LOG_CATEGORY_STATIC(LogGoogleARCoreImageBlitter, Log, All);
+
+FGoogleARCoreDeviceCameraBlitter::FGoogleARCoreDeviceCameraBlitter()
+  : CurrentCameraCopy(0)
+  , BlitShaderProgram(0)
+  , BlitShaderProgram_Uniform_CameraTexture(0)
+  , BlitShaderProgram_Attribute_InPos(0)
+  , FrameBufferObject(0)
+  , VertexBufferObject(0)
+{
+}
+
+FGoogleARCoreDeviceCameraBlitter::~FGoogleARCoreDeviceCameraBlitter()
+{
+	for(int32 i = 0; i < CameraCopies.Num(); i++)
+	{
+		delete CameraCopies[i];
+	}
+}
+
+UTexture *FGoogleARCoreDeviceCameraBlitter::GetLastCameraImageTexture()
+{
+	if(CameraCopies.Num())
+	{
+		int32 LastCameraImage = static_cast<int32>(CurrentCameraCopy) - 1;
+		if(LastCameraImage < 0)
+		{
+			LastCameraImage = CameraCopies.Num() - 1;
+		}
+		return CameraCopies[LastCameraImage];
+	}
+	return nullptr;
+}
+
+void FGoogleARCoreDeviceCameraBlitter::LateInit(FIntPoint ImageSize)
+{
+#if PLATFORM_ANDROID
+	FIntPoint CameraSize = ImageSize;
+
+	// Create a ring buffer of UTexture2Ds to store snapshots of
+	// the camera texture, if we don't already have them.
+	if(!CameraCopies.Num())
+	{
+		for(uint32 i = 0; i < 5; i++)
+		{
+			// Make the texture itself.
+			UTexture *CameraCopy = UTexture2D::CreateTransient(CameraSize.X, CameraSize.Y, PF_R8G8B8A8);
+			check(CameraCopy);
+			CameraCopy->AddToRoot();
+			FTextureResource *resource = CameraCopy->CreateResource();
+			CameraCopy->Filter = TextureFilter::TF_Nearest;
+			CameraCopy->SRGB = false;
+			CameraCopy->UpdateResource();
+			CameraCopies.Add(CameraCopy);
+
+			// Make a place to store the OpenGL texture ID, read
+			// back from the RHI within the render thread.
+			uint32 *TextureId = new uint32;
+			*TextureId = 0;
+			CameraCopyIds.Add(TextureId);
+
+			// We have to get the OpenGL texture ID on the render
+			// thread, after the native resource has been
+			// initialized.
+			ENQUEUE_UNIQUE_RENDER_COMMAND_THREEPARAMETER(
+				InitCameraBlitter,
+				FTexture2DResource *, resource, ((FTexture2DResource *)resource),
+				uint32 *, TextureId, TextureId,
+				FIntPoint, CameraSize, CameraSize,
+				{
+					resource->InitRHI();
+					*TextureId = *reinterpret_cast<uint32*>(resource->GetTexture2DRHI()->GetNativeResource());
+				});
+		}
+	}
+
+	// Compile and link the shader program.
+	if(!BlitShaderProgram)
+	{
+		GLint CompileSuccess = 0;
+
+		// Compile VS.
+		const char *VsText =
+			"#version 300 es\n"
+			"#extension GL_OES_EGL_image_external_essl3 : require\n"
+			"precision mediump float;\n"
+			"in vec2 inPos;\n"
+			"out vec2 uvs;\n"
+			"void main() {\n"
+			"	uvs = inPos;\n"
+			"	vec2 clipSpace = inPos * 2.0 - vec2(1.0, 1.0);\n"
+			"	gl_Position = vec4(clipSpace, 0.0, 1.0);\n"
+			"}\n";
+		GLuint Vs = glCreateShader(GL_VERTEX_SHADER);
+		glShaderSource(Vs, 1, &VsText, NULL);
+		glCompileShader(Vs);
+		glGetShaderiv(Vs, GL_COMPILE_STATUS, &CompileSuccess);
+		if(!CompileSuccess) {
+			GLint ShaderLogLength = 0;
+			glGetShaderiv(Vs, GL_INFO_LOG_LENGTH, &ShaderLogLength);
+			char *ShaderLog = new char[ShaderLogLength + 1];
+			glGetShaderInfoLog(Vs, ShaderLogLength, &ShaderLogLength, ShaderLog);
+			ShaderLog[ShaderLogLength] = 0;
+			UE_LOG(LogGoogleARCoreImageBlitter, Error, TEXT("Shader compile error: %s"), *FString(ShaderLog));
+			check(CompileSuccess);
+		}
+
+		// Compile FS.
+		const char *FsText =
+			"#version 300 es\n"
+			"#extension GL_OES_EGL_image_external_essl3 : require\n"
+			"precision mediump float;\n"
+			"uniform samplerExternalOES cameraTexture;\n"
+			"in vec2 uvs;\n"
+			"layout(location = 0) out vec4 outColor;\n"
+			"void main() {\n"
+			"	outColor.xyz = texture(cameraTexture, uvs).rgb;\n"
+			"	outColor.w = 1.0;\n"
+			"}\n";
+		GLuint Fs = glCreateShader(GL_FRAGMENT_SHADER);
+		glShaderSource(Fs, 1, &FsText, NULL);
+		glCompileShader(Fs);
+		glGetShaderiv(Fs, GL_COMPILE_STATUS, &CompileSuccess);
+		if(!CompileSuccess)
+		{
+			GLint ShaderLogLength = 0;
+			glGetShaderiv(Fs, GL_INFO_LOG_LENGTH, &ShaderLogLength);
+			char *ShaderLog = new char[ShaderLogLength + 1];
+			glGetShaderInfoLog(Fs, ShaderLogLength, &ShaderLogLength, ShaderLog);
+			ShaderLog[ShaderLogLength] = 0;
+			UE_LOG(LogGoogleARCoreImageBlitter, Error, TEXT("Shader compile error: %s"), *FString(ShaderLog));
+			check(CompileSuccess);
+		}
+
+		// Link program.
+		BlitShaderProgram = glCreateProgram();
+		glAttachShader(BlitShaderProgram, Vs);
+		glAttachShader(BlitShaderProgram, Fs);
+		glLinkProgram(BlitShaderProgram);
+
+		glGetProgramiv(BlitShaderProgram, GL_LINK_STATUS, &CompileSuccess);
+		if(!CompileSuccess)
+		{
+			GLint ShaderLogLength = 0;
+			glGetProgramiv(BlitShaderProgram, GL_INFO_LOG_LENGTH, &ShaderLogLength);
+			char *ShaderLog = new char[ShaderLogLength + 1];
+			glGetProgramInfoLog(Fs, ShaderLogLength, &ShaderLogLength, ShaderLog);
+			ShaderLog[ShaderLogLength] = 0;
+			UE_LOG(LogGoogleARCoreImageBlitter, Error, TEXT("Shader link error: %s"), *FString(ShaderLog));
+			check(CompileSuccess);
+		}
+
+		// Get uniform and attribute locations.
+		BlitShaderProgram_Uniform_CameraTexture = glGetUniformLocation(BlitShaderProgram, "cameraTexture");
+		BlitShaderProgram_Attribute_InPos = glGetAttribLocation(BlitShaderProgram, "inPos");
+
+		// Clean up unneeded references.
+		glDeleteShader(Vs);
+		glDeleteShader(Fs);
+
+		// Sanity check.
+		check(glGetError() == GL_NO_ERROR);
+	}
+
+	// Create FBO.
+	if(!FrameBufferObject)
+	{
+		glGenFramebuffers(1, &FrameBufferObject);
+	}
+
+	// Create VBO.
+	if(!VertexBufferObject)
+	{
+		glGenBuffers(1, &VertexBufferObject);
+	}
+#endif
+}
+
+void FGoogleARCoreDeviceCameraBlitter::DoBlit(uint32_t CameraTextureId, FIntPoint ImageSize)
+{
+#if PLATFORM_ANDROID
+	FIntPoint CameraSize = ImageSize;
+
+	LateInit(ImageSize);
+
+	if (CameraCopyIds.Num())
+	{
+		uint32 TextureId = *CameraCopyIds[CurrentCameraCopy];
+
+		if (TextureId)
+		{
+			glBindFramebuffer(GL_FRAMEBUFFER, FrameBufferObject);
+			glFramebufferTexture2D(GL_FRAMEBUFFER, GL_COLOR_ATTACHMENT0, GL_TEXTURE_2D, TextureId, 0);
+
+			glViewport(0, 0, CameraSize.X, CameraSize.Y);
+
+			glDisable(GL_DEPTH_TEST);
+			glDisable(GL_BLEND);
+			glDisable(GL_STENCIL_TEST);
+
+			glUseProgram(BlitShaderProgram);
+
+			glEnableVertexAttribArray(BlitShaderProgram_Attribute_InPos);
+
+			float Verts[6 * 2] = {
+				0.0f, 1.0f,
+				1.0f, 1.0f,
+				1.0f, 0.0f,
+
+				0.0f, 1.0f,
+				1.0f, 0.0f,
+				0.0f, 0.0f
+			};
+
+			glBindBuffer(GL_ARRAY_BUFFER, VertexBufferObject);
+			glBufferData(GL_ARRAY_BUFFER, sizeof(Verts), Verts, GL_STATIC_DRAW);
+
+			glVertexAttribPointer(BlitShaderProgram_Attribute_InPos, 2, GL_FLOAT, false, 0, 0);
+
+			glActiveTexture(GL_TEXTURE0);
+			glBindTexture(GL_TEXTURE_EXTERNAL_OES, CameraTextureId);
+			glUniform1i(BlitShaderProgram_Uniform_CameraTexture, 0);
+
+			glDrawArrays(GL_TRIANGLES, 0, 6);
+
+			glBindTexture(GL_TEXTURE_EXTERNAL_OES, 0);
+			glDisableVertexAttribArray(BlitShaderProgram_Attribute_InPos);
+			glFramebufferTexture2D(GL_FRAMEBUFFER, GL_COLOR_ATTACHMENT0, GL_TEXTURE_2D, 0, 0);
+			glBindFramebuffer(GL_FRAMEBUFFER, 0);
+
+			glFlush();
+
+		}
+
+		// Update the dynamic material instance with the newest
+		// texture.
+		UMaterialInstanceDynamic *DynMat = Cast<UMaterialInstanceDynamic>(
+			GetDefault<UGoogleARCoreCameraOverlayMaterialLoader>()->DefaultCameraOverlayMaterial);
+		DynMat->SetTextureParameterValue(FName("CameraTexture"), CameraCopies[CurrentCameraCopy]);
+
+		CurrentCameraCopy++;
+		CurrentCameraCopy %= CameraCopies.Num();
+	}
+#endif
+}
diff --git a/Engine/Plugins/Experimental/AR/Google/GoogleARCore/Source/GoogleARCoreBase/Private/GoogleARCoreCameraImageBlitter.h b/Engine/Plugins/Experimental/AR/Google/GoogleARCore/Source/GoogleARCoreBase/Private/GoogleARCoreCameraImageBlitter.h
new file mode 100644
--- /dev/null
+++ b/Engine/Plugins/Experimental/AR/Google/GoogleARCore/Source/GoogleARCoreBase/Private/GoogleARCoreCameraImageBlitter.h
@@ -0,0 +1,32 @@
+// Copyright 2017 Google Inc.
+
+#pragma once
+
+#include "Misc/CoreDelegates.h"
+#include "Engine/EngineBaseTypes.h"
+#include "Engine/Texture.h"
+#include "HAL/ThreadSafeBool.h"
+
+class FGoogleARCoreDeviceCameraBlitter
+{
+public:
+
+	FGoogleARCoreDeviceCameraBlitter();
+	~FGoogleARCoreDeviceCameraBlitter();
+
+	void DoBlit(uint32_t TextureId, FIntPoint ImageSize);
+	UTexture *GetLastCameraImageTexture();
+
+private:
+
+	void LateInit(FIntPoint ImageSize);
+
+	uint32 CurrentCameraCopy;
+	uint32 BlitShaderProgram;
+	uint32 BlitShaderProgram_Uniform_CameraTexture;
+	uint32 BlitShaderProgram_Attribute_InPos;
+	uint32 FrameBufferObject;
+	uint32 VertexBufferObject;
+	TArray<UTexture *> CameraCopies;
+    TArray<uint32 *> CameraCopyIds;
+};
diff --git a/Engine/Plugins/Experimental/AR/Google/GoogleARCore/Source/GoogleARCoreBase/Private/GoogleARCoreCameraIntrinsics.cpp b/Engine/Plugins/Experimental/AR/Google/GoogleARCore/Source/GoogleARCoreBase/Private/GoogleARCoreCameraIntrinsics.cpp
new file mode 100644
--- /dev/null
+++ b/Engine/Plugins/Experimental/AR/Google/GoogleARCore/Source/GoogleARCoreBase/Private/GoogleARCoreCameraIntrinsics.cpp
@@ -0,0 +1,58 @@
+// Copyright 2018 Google Inc.
+
+#include "GoogleARCoreCameraIntrinsics.h"
+
+#if PLATFORM_ANDROID
+#include "arcore_c_api.h"
+#endif
+
+#include "GoogleARCoreAPI.h"
+
+UGoogleARCoreCameraIntrinsics::~UGoogleARCoreCameraIntrinsics()
+{
+#if PLATFORM_ANDROID
+	if (NativeCameraIntrinsics)
+	{
+		ArCameraIntrinsics_destroy(NativeCameraIntrinsics);
+	}
+#endif
+}
+
+void UGoogleARCoreCameraIntrinsics::GetFocalLength(float &OutFX, float &OutFY)
+{
+#if PLATFORM_ANDROID
+	TSharedPtr<FGoogleARCoreSession> SessionPtr = Session.Pin();
+	if (NativeCameraIntrinsics && SessionPtr.IsValid())
+	{
+		ArCameraIntrinsics_getFocalLength(
+			SessionPtr->GetHandle(), NativeCameraIntrinsics,
+			&OutFX, &OutFY);
+	}
+#endif
+}
+
+void UGoogleARCoreCameraIntrinsics::GetPrincipalPoint(float &OutCX, float &OutCY)
+{
+#if PLATFORM_ANDROID
+	TSharedPtr<FGoogleARCoreSession> SessionPtr = Session.Pin();
+	if (NativeCameraIntrinsics && SessionPtr.IsValid())
+	{
+		ArCameraIntrinsics_getPrincipalPoint(
+			SessionPtr->GetHandle(), NativeCameraIntrinsics,
+			&OutCX, &OutCY);
+	}
+#endif
+}
+
+void UGoogleARCoreCameraIntrinsics::GetImageDimensions(int32 &OutWidth, int32 &OutHeight)
+{
+#if PLATFORM_ANDROID
+	TSharedPtr<FGoogleARCoreSession> SessionPtr = Session.Pin();
+	if (NativeCameraIntrinsics && SessionPtr.IsValid())
+	{
+		ArCameraIntrinsics_getImageDimensions(
+			SessionPtr->GetHandle(), NativeCameraIntrinsics,
+			&OutWidth, &OutHeight);
+	}
+#endif
+}
diff --git a/Engine/Plugins/Experimental/AR/Google/GoogleARCore/Source/GoogleARCoreBase/Private/GoogleARCoreDevice.cpp b/Engine/Plugins/Experimental/AR/Google/GoogleARCore/Source/GoogleARCoreBase/Private/GoogleARCoreDevice.cpp
new file mode 100644
--- /dev/null
+++ b/Engine/Plugins/Experimental/AR/Google/GoogleARCore/Source/GoogleARCoreBase/Private/GoogleARCoreDevice.cpp
@@ -0,0 +1,732 @@
+// Copyright 2017 Google Inc.
+
+#include "GoogleARCoreDevice.h"
+#include "CoreMinimal.h"
+#include "Misc/ScopeLock.h"
+#include "HAL/ThreadSafeCounter64.h"
+#include "GameFramework/WorldSettings.h"
+#include "Framework/Application/SlateApplication.h"
+#include "Engine/World.h" // for FWorldDelegates
+#include "Engine/Engine.h"
+#include "GeneralProjectSettings.h"
+#include "GoogleARCoreXRTrackingSystem.h"
+#include "GoogleARCoreAndroidHelper.h"
+#include "GoogleARCoreBaseLogCategory.h"
+
+#if PLATFORM_ANDROID
+#include "Android/AndroidApplication.h"
+#endif
+
+#include "GoogleARCorePermissionHandler.h"
+
+namespace
+{
+	EGoogleARCoreFunctionStatus ToARCoreFunctionStatus(EGoogleARCoreAPIStatus Status)
+	{
+		switch (Status)
+		{
+		case EGoogleARCoreAPIStatus::AR_SUCCESS:
+			return EGoogleARCoreFunctionStatus::Success;
+		case EGoogleARCoreAPIStatus::AR_ERROR_NOT_TRACKING:
+			return EGoogleARCoreFunctionStatus::NotTracking;
+		case EGoogleARCoreAPIStatus::AR_ERROR_SESSION_PAUSED:
+			return EGoogleARCoreFunctionStatus::SessionPaused;
+		case EGoogleARCoreAPIStatus::AR_ERROR_RESOURCE_EXHAUSTED:
+			return EGoogleARCoreFunctionStatus::ResourceExhausted;
+		case EGoogleARCoreAPIStatus::AR_ERROR_NOT_YET_AVAILABLE:
+			return EGoogleARCoreFunctionStatus::NotAvailable;
+		default:
+			ensureMsgf(false, TEXT("Unknown conversion from EGoogleARCoreAPIStatus %d to EGoogleARCoreFunctionStatus."), static_cast<int>(Status));
+			return EGoogleARCoreFunctionStatus::Unknown;
+		}
+	}
+}
+
+FGoogleARCoreDelegates::FGoogleARCoreOnConfigCameraDelegate FGoogleARCoreDelegates::OnCameraConfig;
+
+FGoogleARCoreDevice* FGoogleARCoreDevice::GetInstance()
+{
+	static FGoogleARCoreDevice Inst;
+	return &Inst;
+}
+
+FGoogleARCoreDevice::FGoogleARCoreDevice()
+	: PassthroughCameraTexture(nullptr)
+	, PassthroughCameraTextureId(-1)
+	, bIsARCoreSessionRunning(false)
+	, bForceLateUpdateEnabled(false)
+	, bSessionConfigChanged(false)
+	, bAndroidRuntimePermissionsRequested(false)
+	, bAndroidRuntimePermissionsGranted(false)
+	, bPermissionDeniedByUser(false)
+	, bStartSessionRequested(false)
+	, bShouldSessionRestart(false)
+	, bARCoreInstallRequested(false)
+	, bARCoreInstalled(false)
+	, WorldToMeterScale(100.0f)
+	, PermissionHandler(nullptr)
+	, bDisplayOrientationChanged(false)
+	, CurrentSessionStatus(EARSessionStatus::NotStarted)
+{
+}
+
+EGoogleARCoreAvailability FGoogleARCoreDevice::CheckARCoreAPKAvailability()
+{
+	return FGoogleARCoreAPKManager::CheckARCoreAPKAvailability();
+}
+
+EGoogleARCoreAPIStatus FGoogleARCoreDevice::RequestInstall(bool bUserRequestedInstall, EGoogleARCoreInstallStatus& OutInstallStatus)
+{
+	return FGoogleARCoreAPKManager::RequestInstall(bUserRequestedInstall, OutInstallStatus);
+}
+
+bool FGoogleARCoreDevice::GetIsTrackingTypeSupported(EARSessionType SessionType)
+{
+	if (SessionType == EARSessionType::World)
+	{
+		return true;
+	}
+	else
+	{
+		return false;
+	}
+}
+
+void FGoogleARCoreDevice::OnModuleLoaded()
+{
+	// Init display orientation.
+	OnDisplayOrientationChanged();
+
+	FWorldDelegates::OnWorldTickStart.AddRaw(this, &FGoogleARCoreDevice::OnWorldTickStart);
+}
+
+void FGoogleARCoreDevice::OnModuleUnloaded()
+{
+	FWorldDelegates::OnWorldTickStart.RemoveAll(this);
+	// clear the unique ptr.
+	ARCoreSession.Reset();
+}
+
+bool FGoogleARCoreDevice::GetIsARCoreSessionRunning()
+{
+	return bIsARCoreSessionRunning;
+}
+
+EARSessionStatus FGoogleARCoreDevice::GetSessionStatus()
+{
+	return CurrentSessionStatus;
+}
+
+float FGoogleARCoreDevice::GetWorldToMetersScale()
+{
+	return WorldToMeterScale;
+}
+
+// This function will be called by public function to start AR core session request.
+void FGoogleARCoreDevice::StartARCoreSessionRequest(UARSessionConfig* SessionConfig)
+{
+	UE_LOG(LogGoogleARCore, Log, TEXT("Start ARCore session requested."));
+
+	if (bIsARCoreSessionRunning)
+	{
+		UE_LOG(LogGoogleARCore, Log, TEXT("ARCore session is already running, set it to use the new session config."));
+		EGoogleARCoreAPIStatus Status = ARCoreSession->ConfigSession(*SessionConfig);
+		ensureMsgf(Status == EGoogleARCoreAPIStatus::AR_SUCCESS, TEXT("Failed to set ARCore session to new configuration while it is running."));
+
+		return;
+	}
+
+	if (bStartSessionRequested)
+	{
+		UE_LOG(LogGoogleARCore, Warning, TEXT("ARCore session is already starting. This will overriding the previous session config with the new one."))
+	}
+
+	bStartSessionRequested = true;
+	// Re-request permission if necessary
+	bPermissionDeniedByUser = false;
+	bARCoreInstallRequested = false;
+
+	// Try recreating the ARCoreSession to fix the fatal error.
+	if (CurrentSessionStatus == EARSessionStatus::FatalError)
+	{
+		UE_LOG(LogGoogleARCore, Warning, TEXT("Reset ARCore session due to fatal error detected."));
+		ResetARCoreSession();
+	}
+}
+
+bool FGoogleARCoreDevice::SetARCameraConfig(FGoogleARCoreCameraConfig CameraConfig)
+{
+	EGoogleARCoreAPIStatus APIStatus = ARCoreSession->SetCameraConfig(CameraConfig);
+	if (APIStatus == EGoogleARCoreAPIStatus::AR_SUCCESS)
+	{
+		return true;
+	}
+	else if (APIStatus == EGoogleARCoreAPIStatus::AR_ERROR_SESSION_NOT_PAUSED)
+	{
+		UE_LOG(LogGoogleARCore, Error, TEXT("Failed to set ARCamera configuration due to the arcore session isn't paused."));
+	}
+	else
+	{
+		UE_LOG(LogGoogleARCore, Error, TEXT("Failed to set ARCamera configuration with provided CameraConfig."));
+	}
+
+	UE_LOG(LogGoogleARCore, Error, TEXT("You should only call the ConfigARCoreCamera function when the OnConfigCamera delegate get called, and the provided CameraConfig must be from the array that is passed by the delegate."));
+	return false;
+}
+
+bool FGoogleARCoreDevice::GetARCameraConfig(FGoogleARCoreCameraConfig& OutCurrentCameraConfig)
+{
+	if (ARCoreSession.IsValid())
+	{
+		ARCoreSession->GetARCameraConfig(OutCurrentCameraConfig);
+		return true;
+	}
+	else
+	{
+		return false;
+	}
+}
+
+bool FGoogleARCoreDevice::GetStartSessionRequestFinished()
+{
+	return !bStartSessionRequested;
+}
+
+// Note that this function will only be registered when ARCore is supported.
+void FGoogleARCoreDevice::OnWorldTickStart(ELevelTick TickType, float DeltaTime)
+{
+	WorldToMeterScale = GWorld->GetWorldSettings()->WorldToMeters;
+	TFunction<void()> Func;
+	while (RunOnGameThreadQueue.Dequeue(Func))
+	{
+		Func();
+	}
+
+	if (!bIsARCoreSessionRunning && bStartSessionRequested)
+	{
+		if (!bARCoreInstalled)
+		{
+			EGoogleARCoreInstallStatus InstallStatus = EGoogleARCoreInstallStatus::Installed;
+			EGoogleARCoreAPIStatus Status = FGoogleARCoreAPKManager::RequestInstall(!bARCoreInstallRequested, InstallStatus);
+
+			if (Status != EGoogleARCoreAPIStatus::AR_SUCCESS)
+			{
+				bStartSessionRequested = false;
+				CurrentSessionStatus = EARSessionStatus::NotSupported;
+			}
+			else if (InstallStatus == EGoogleARCoreInstallStatus::Installed)
+			{
+				bARCoreInstalled = true;
+			}
+			else
+			{
+				bARCoreInstallRequested = true;
+			}
+		}
+
+		else if (bPermissionDeniedByUser)
+		{
+			CurrentSessionStatus = EARSessionStatus::PermissionNotGranted;
+			bStartSessionRequested = false;
+		}
+		else
+		{
+			CheckAndRequrestPermission(*AccessSessionConfig());
+			// Either we don't need to request permission or the permission request is done.
+			// Queue the session start task on UiThread
+			if (!bAndroidRuntimePermissionsRequested)
+			{
+				StartSessionWithRequestedConfig();
+			}
+		}
+	}
+
+	if (bIsARCoreSessionRunning)
+	{
+		// Update ARFrame
+		FVector2D ViewportSize(1, 1);
+		if (GEngine && GEngine->GameViewport)
+		{
+			ViewportSize = GEngine->GameViewport->Viewport->GetSizeXY();
+		}
+		ARCoreSession->SetDisplayGeometry(static_cast<int>(FGoogleARCoreAndroidHelper::GetDisplayRotation()), ViewportSize.X, ViewportSize.Y);
+		EGoogleARCoreAPIStatus Status = ARCoreSession->Update(WorldToMeterScale);
+		if (Status == EGoogleARCoreAPIStatus::AR_ERROR_FATAL)
+		{
+			ARCoreSession->Pause();
+			bIsARCoreSessionRunning = false;
+			CurrentSessionStatus = EARSessionStatus::FatalError;
+		}
+		else
+		{
+			// Copy the camera GPU texture to a normal Texture2D so that we can get around the multi-threaded rendering issue.
+			CameraBlitter.DoBlit(PassthroughCameraTextureId, SessionCameraConfig.CameraTextureResolution);
+		}
+	}
+}
+
+void FGoogleARCoreDevice::CheckAndRequrestPermission(const UARSessionConfig& ConfigurationData)
+{
+	if (!bAndroidRuntimePermissionsRequested)
+	{
+		TArray<FString> RuntimePermissions;
+		TArray<FString> NeededPermissions;
+		GetRequiredRuntimePermissionsForConfiguration(ConfigurationData, RuntimePermissions);
+		if (RuntimePermissions.Num() > 0)
+		{
+			for (int32 i = 0; i < RuntimePermissions.Num(); i++)
+			{
+				if (!UARCoreAndroidPermissionHandler::CheckRuntimePermission(RuntimePermissions[i]))
+				{
+					NeededPermissions.Add(RuntimePermissions[i]);
+				}
+			}
+		}
+		if (NeededPermissions.Num() > 0)
+		{
+			bAndroidRuntimePermissionsGranted = false;
+			bAndroidRuntimePermissionsRequested = true;
+			if (PermissionHandler == nullptr)
+			{
+				PermissionHandler = NewObject<UARCoreAndroidPermissionHandler>();
+				PermissionHandler->AddToRoot();
+			}
+			PermissionHandler->RequestRuntimePermissions(NeededPermissions);
+		}
+		else
+		{
+			bAndroidRuntimePermissionsGranted = true;
+		}
+	}
+}
+
+void FGoogleARCoreDevice::HandleRuntimePermissionsGranted(const TArray<FString>& RuntimePermissions, const TArray<bool>& Granted)
+{
+	bool bGranted = true;
+	for (int32 i = 0; i < RuntimePermissions.Num(); i++)
+	{
+		if (!Granted[i])
+		{
+			bGranted = false;
+			UE_LOG(LogGoogleARCore, Warning, TEXT("Android runtime permission denied: %s"), *RuntimePermissions[i]);
+		}
+		else
+		{
+			UE_LOG(LogGoogleARCore, Log, TEXT("Android runtime permission granted: %s"), *RuntimePermissions[i]);
+		}
+	}
+	bAndroidRuntimePermissionsRequested = false;
+	bAndroidRuntimePermissionsGranted = bGranted;
+
+	if (!bGranted)
+	{
+		bPermissionDeniedByUser = true;
+	}
+}
+
+void FGoogleARCoreDevice::StartSessionWithRequestedConfig()
+{
+	bStartSessionRequested = false;
+
+	// Allocate passthrough camera texture if necessary.
+	if (PassthroughCameraTexture == nullptr)
+	{
+		ENQUEUE_UNIQUE_RENDER_COMMAND_ONEPARAMETER(
+			UpdateCameraImageUV,
+			FGoogleARCoreDevice*, ARCoreDevicePtr, this,
+			{
+				ARCoreDevicePtr->AllocatePassthroughCameraTexture_RenderThread();
+			}
+		);
+		FlushRenderingCommands();
+	}
+
+	if (!ARCoreSession.IsValid())
+	{
+		ARCoreSession = FGoogleARCoreSession::CreateARCoreSession();
+		EGoogleARCoreAPIStatus SessionCreateStatus = ARCoreSession->GetSessionCreateStatus();
+		if (SessionCreateStatus != EGoogleARCoreAPIStatus::AR_SUCCESS)
+		{
+			ensureMsgf(false, TEXT("Failed to create ARCore session with error status: %d"), (int)SessionCreateStatus);
+			if (SessionCreateStatus != EGoogleARCoreAPIStatus::AR_ERROR_FATAL)
+			{
+				CurrentSessionStatus = EARSessionStatus::NotSupported;
+			}
+			else
+			{
+				CurrentSessionStatus = EARSessionStatus::FatalError;
+			}
+			ARCoreSession.Reset();
+			return;
+		}
+		ARCoreSession->SetARSystem(ARSystem.ToSharedRef());
+	}
+
+	StartSession();
+}
+
+void FGoogleARCoreDevice::StartSession()
+{
+	UARSessionConfig* RequestedConfig = AccessSessionConfig();
+
+	if (RequestedConfig->GetSessionType() != EARSessionType::World)
+	{
+		UE_LOG(LogGoogleARCore, Warning, TEXT("Start AR failed: Unsupported AR tracking type %d for GoogleARCore"), static_cast<int>(RequestedConfig->GetSessionType()));
+		CurrentSessionStatus = EARSessionStatus::UnsupportedConfiguration;
+		return;
+	}
+
+	EGoogleARCoreAPIStatus Status = ARCoreSession->ConfigSession(*RequestedConfig);
+
+	if (Status != EGoogleARCoreAPIStatus::AR_SUCCESS)
+	{
+		UE_LOG(LogGoogleARCore, Error, TEXT("ARCore Session start failed with error status %d"), static_cast<int>(Status));
+		CurrentSessionStatus = EARSessionStatus::UnsupportedConfiguration;
+		return;
+	}
+
+	check(PassthroughCameraTextureId != -1);
+	ARCoreSession->SetCameraTextureId(PassthroughCameraTextureId);
+
+	FGoogleARCoreDelegates::OnCameraConfig.Broadcast(ARCoreSession->GetSupportedCameraConfig());
+
+	Status = ARCoreSession->Resume();
+
+	if (Status != EGoogleARCoreAPIStatus::AR_SUCCESS)
+	{
+		UE_LOG(LogGoogleARCore, Error, TEXT("ARCore Session start failed with error status %d"), static_cast<int>(Status));
+		check(Status == EGoogleARCoreAPIStatus::AR_ERROR_FATAL);
+		// If we failed here, the only reason would be fatal error.
+		CurrentSessionStatus = EARSessionStatus::FatalError;
+		return;
+	}
+
+	if (GEngine->XRSystem.IsValid())
+	{
+		FGoogleARCoreXRTrackingSystem* ARCoreTrackingSystem = static_cast<FGoogleARCoreXRTrackingSystem*>(GEngine->XRSystem.Get());
+		if (ARCoreTrackingSystem)
+		{
+			const bool bMatchFOV = RequestedConfig->ShouldRenderCameraOverlay();
+			ARCoreTrackingSystem->ConfigARCoreXRCamera(bMatchFOV, RequestedConfig->ShouldRenderCameraOverlay());
+		}
+		else
+		{
+			UE_LOG(LogGoogleARCore, Error, TEXT("ERROR: GoogleARCoreXRTrackingSystem is not available."));
+		}
+	}
+
+	ARCoreSession->GetARCameraConfig(SessionCameraConfig);
+
+	bIsARCoreSessionRunning = true;
+	CurrentSessionStatus = EARSessionStatus::Running;
+	UE_LOG(LogGoogleARCore, Log, TEXT("ARCore session started successfully."));
+
+	ARSystem->OnARSessionStarted.Broadcast();
+}
+
+void FGoogleARCoreDevice::SetARSystem(TSharedPtr<FARSupportInterface , ESPMode::ThreadSafe> InARSystem)
+{
+	check(InARSystem.IsValid());
+	ARSystem = InARSystem;
+}
+
+void* FGoogleARCoreDevice::GetARSessionRawPointer()
+{
+#if PLATFORM_ANDROID
+	return reinterpret_cast<void*>(ARCoreSession->GetHandle());
+#endif
+	return nullptr;
+}
+
+void* FGoogleARCoreDevice::GetGameThreadARFrameRawPointer()
+{
+#if PLATFORM_ANDROID
+	return reinterpret_cast<void*>(ARCoreSession->GetLatestFrameRawPointer());
+#endif
+	return nullptr;
+}
+
+TSharedPtr<FARSupportInterface , ESPMode::ThreadSafe> FGoogleARCoreDevice::GetARSystem()
+{
+	return ARSystem;
+}
+
+void FGoogleARCoreDevice::PauseARCoreSession()
+{
+	UE_LOG(LogGoogleARCore, Log, TEXT("Pausing ARCore session."));
+	if (!bIsARCoreSessionRunning)
+	{
+		if(bStartSessionRequested)
+		{
+			bStartSessionRequested = false;
+		}
+		else
+		{
+			UE_LOG(LogGoogleARCore, Log, TEXT("Could not stop ARCore tracking session because there is no running tracking session!"));
+		}
+		return;
+	}
+
+	EGoogleARCoreAPIStatus Status = ARCoreSession->Pause();
+
+	if (Status == EGoogleARCoreAPIStatus::AR_ERROR_FATAL)
+	{
+		CurrentSessionStatus = EARSessionStatus::FatalError;
+	}
+
+	bIsARCoreSessionRunning = false;
+	CurrentSessionStatus = EARSessionStatus::NotStarted;
+	UE_LOG(LogGoogleARCore, Log, TEXT("ARCore session paused"));
+}
+
+void FGoogleARCoreDevice::ResetARCoreSession()
+{
+	ARCoreSession.Reset();
+	CurrentSessionStatus = EARSessionStatus::NotStarted;
+}
+
+void FGoogleARCoreDevice::AllocatePassthroughCameraTexture_RenderThread()
+{
+	FRHICommandListImmediate& RHICmdList = FRHICommandListExecutor::GetImmediateCommandList();
+	FRHIResourceCreateInfo CreateInfo;
+
+	PassthroughCameraTexture = RHICmdList.CreateTextureExternal2D(1, 1, PF_R8G8B8A8, 1, 1, 0, CreateInfo);
+
+	void* NativeResource = PassthroughCameraTexture->GetNativeResource();
+	check(NativeResource);
+	PassthroughCameraTextureId = *reinterpret_cast<uint32*>(NativeResource);
+}
+
+FTextureRHIRef FGoogleARCoreDevice::GetPassthroughCameraTexture()
+{
+	return PassthroughCameraTexture;
+}
+
+FMatrix FGoogleARCoreDevice::GetPassthroughCameraProjectionMatrix(FIntPoint ViewRectSize) const
+{
+	if (!IsARSessionInitialized())
+	{
+		return FMatrix::Identity;
+	}
+	return ARCoreSession->GetLatestFrame()->GetProjectionMatrix();
+}
+
+void FGoogleARCoreDevice::GetPassthroughCameraImageUVs(const TArray<float>& InUvs, TArray<float>& OutUVs) const
+{
+	if (!IsARSessionInitialized())
+	{
+		return;
+	}
+	ARCoreSession->GetLatestFrame()->TransformDisplayUvCoords(InUvs, OutUVs);
+}
+
+EGoogleARCoreTrackingState FGoogleARCoreDevice::GetTrackingState() const
+{
+	if (!IsARSessionInitialized())
+	{
+		return EGoogleARCoreTrackingState::StoppedTracking;
+	}
+	else if (!bIsARCoreSessionRunning)
+	{
+		return EGoogleARCoreTrackingState::NotTracking;
+	}
+
+	return ARCoreSession->GetLatestFrame()->GetCameraTrackingState();
+}
+
+FTransform FGoogleARCoreDevice::GetLatestPose() const
+{
+	if (!IsARSessionInitialized())
+	{
+		return FTransform::Identity;
+	}
+	return ARCoreSession->GetLatestFrame()->GetCameraPose();
+}
+
+EGoogleARCoreFunctionStatus FGoogleARCoreDevice::GetLatestPointCloud(UGoogleARCorePointCloud*& OutLatestPointCloud) const
+{
+	if (!IsARSessionInitialized())
+	{
+		return EGoogleARCoreFunctionStatus::NotAvailable;
+	}
+
+	return ToARCoreFunctionStatus(ARCoreSession->GetLatestFrame()->GetPointCloud(OutLatestPointCloud));
+}
+
+EGoogleARCoreFunctionStatus FGoogleARCoreDevice::AcquireLatestPointCloud(UGoogleARCorePointCloud*& OutLatestPointCloud) const
+{
+	if (!IsARSessionInitialized())
+	{
+		return EGoogleARCoreFunctionStatus::NotAvailable;
+	}
+
+	return ToARCoreFunctionStatus(ARCoreSession->GetLatestFrame()->AcquirePointCloud(OutLatestPointCloud));
+}
+
+#if PLATFORM_ANDROID
+EGoogleARCoreFunctionStatus FGoogleARCoreDevice::GetLatestCameraMetadata(const ACameraMetadata*& OutCameraMetadata) const
+{
+	if (!IsARSessionInitialized())
+	{
+		return EGoogleARCoreFunctionStatus::NotAvailable;
+	}
+
+	return ToARCoreFunctionStatus(ARCoreSession->GetLatestFrame()->GetCameraMetadata(OutCameraMetadata));
+}
+#endif
+
+EGoogleARCoreFunctionStatus FGoogleARCoreDevice::AcquireCameraImage(UGoogleARCoreCameraImage *&OutLatestCameraImage)
+{
+	if (!IsARSessionInitialized())
+	{
+		return EGoogleARCoreFunctionStatus::NotAvailable;
+	}
+
+	return ToARCoreFunctionStatus(ARCoreSession->GetLatestFrame()->AcquireCameraImage(OutLatestCameraImage));
+}
+
+FGoogleARCoreLightEstimate FGoogleARCoreDevice::GetLatestLightEstimate() const
+{
+	if (!IsARSessionInitialized())
+	{
+		return FGoogleARCoreLightEstimate();
+	}
+
+	return ARCoreSession->GetLatestFrame()->GetLightEstimate();
+}
+
+void FGoogleARCoreDevice::ARLineTrace(const FVector2D& ScreenPosition, EGoogleARCoreLineTraceChannel TraceChannels, TArray<FARTraceResult>& OutHitResults)
+{
+	if (!IsARSessionInitialized())
+	{
+		return;
+	}
+	OutHitResults.Empty();
+	ARCoreSession->GetLatestFrame()->ARLineTrace(ScreenPosition, TraceChannels, OutHitResults);
+}
+
+void FGoogleARCoreDevice::ARLineTrace(const FVector& Start, const FVector& End, EGoogleARCoreLineTraceChannel TraceChannels, TArray<FARTraceResult>& OutHitResults)
+{
+	if (!IsARSessionInitialized())
+	{
+		return;
+	}
+	OutHitResults.Empty();
+	ARCoreSession->GetLatestFrame()->ARLineTrace(Start, End, TraceChannels, OutHitResults);
+}
+
+EGoogleARCoreFunctionStatus FGoogleARCoreDevice::CreateARPin(const FTransform& PinToWorldTransform, UARTrackedGeometry* TrackedGeometry, USceneComponent* ComponentToPin, const FName DebugName, UARPin*& OutARAnchorObject)
+{
+	if (!bIsARCoreSessionRunning)
+	{
+		return EGoogleARCoreFunctionStatus::SessionPaused;
+	}
+
+	const FTransform& TrackingToAlignedTracking = ARSystem->GetAlignmentTransform();
+	const FTransform PinToTrackingTransform = PinToWorldTransform.GetRelativeTransform(ARSystem->GetXRTrackingSystem()->GetTrackingToWorldTransform()).GetRelativeTransform(TrackingToAlignedTracking);
+
+	EGoogleARCoreFunctionStatus Status = ToARCoreFunctionStatus(ARCoreSession->CreateARAnchor(PinToTrackingTransform, TrackedGeometry, ComponentToPin, DebugName, OutARAnchorObject));
+
+	return Status;
+}
+
+void FGoogleARCoreDevice::RemoveARPin(UARPin* ARAnchorObject)
+{
+	if (!ARCoreSession.IsValid())
+	{
+		return;
+	}
+
+	ARCoreSession->DetachAnchor(ARAnchorObject);
+}
+
+void FGoogleARCoreDevice::GetAllARPins(TArray<UARPin*>& ARCoreAnchorList)
+{
+	if (!ARCoreSession.IsValid())
+	{
+		return;
+	}
+	ARCoreSession->GetAllAnchors(ARCoreAnchorList);
+}
+
+void FGoogleARCoreDevice::GetUpdatedARPins(TArray<UARPin*>& ARCoreAnchorList)
+{
+	if (!IsARSessionInitialized())
+	{
+		return;
+	}
+	ARCoreSession->GetLatestFrame()->GetUpdatedAnchors(ARCoreAnchorList);
+}
+
+// Functions that are called on Android lifecycle events.
+void FGoogleARCoreDevice::OnApplicationCreated()
+{
+}
+
+void FGoogleARCoreDevice::OnApplicationDestroyed()
+{
+}
+
+void FGoogleARCoreDevice::OnApplicationPause()
+{
+	UE_LOG(LogGoogleARCore, Log, TEXT("OnPause Called: %d"), bIsARCoreSessionRunning);
+	bShouldSessionRestart = bIsARCoreSessionRunning;
+	if (bIsARCoreSessionRunning)
+	{
+		PauseARCoreSession();
+	}
+}
+
+void FGoogleARCoreDevice::OnApplicationResume()
+{
+	UE_LOG(LogGoogleARCore, Log, TEXT("OnResume Called: %d"), bShouldSessionRestart);
+	// Try to ask for permission if it is denied by user.
+	if (bShouldSessionRestart)
+	{
+		bShouldSessionRestart = false;
+		StartSession();
+	}
+}
+
+void FGoogleARCoreDevice::OnApplicationStop()
+{
+}
+
+void FGoogleARCoreDevice::OnApplicationStart()
+{
+}
+
+// TODO: we probably don't need this.
+void FGoogleARCoreDevice::OnDisplayOrientationChanged()
+{
+	FGoogleARCoreAndroidHelper::UpdateDisplayRotation();
+	bDisplayOrientationChanged = true;
+}
+
+UARSessionConfig* FGoogleARCoreDevice::AccessSessionConfig() const
+{
+	return (ARSystem.IsValid())
+		? &ARSystem->AccessSessionConfig()
+		: nullptr;
+}
+
+
+EGoogleARCoreFunctionStatus FGoogleARCoreDevice::GetCameraImageIntrinsics(UGoogleARCoreCameraIntrinsics *&OutCameraIntrinsics)
+{
+	if (!IsARSessionInitialized())
+	{
+		return EGoogleARCoreFunctionStatus::NotAvailable;
+	}
+
+	return ToARCoreFunctionStatus(ARCoreSession->GetLatestFrame()->GetCameraImageIntrinsics(OutCameraIntrinsics));
+}
+
+EGoogleARCoreFunctionStatus FGoogleARCoreDevice::GetCameraTextureIntrinsics(UGoogleARCoreCameraIntrinsics *&OutCameraIntrinsics)
+{
+	if (!IsARSessionInitialized())
+	{
+		return EGoogleARCoreFunctionStatus::NotAvailable;
+	}
+
+	return ToARCoreFunctionStatus(ARCoreSession->GetLatestFrame()->GetCameraTextureIntrinsics(OutCameraIntrinsics));
+}
\ No newline at end of file
diff --git a/Engine/Plugins/Experimental/AR/Google/GoogleARCore/Source/GoogleARCoreBase/Private/GoogleARCoreFunctionLibrary.cpp b/Engine/Plugins/Experimental/AR/Google/GoogleARCore/Source/GoogleARCoreBase/Private/GoogleARCoreFunctionLibrary.cpp
new file mode 100644
--- /dev/null
+++ b/Engine/Plugins/Experimental/AR/Google/GoogleARCore/Source/GoogleARCoreBase/Private/GoogleARCoreFunctionLibrary.cpp
@@ -0,0 +1,408 @@
+// Copyright 2017 Google Inc.
+
+#include "GoogleARCoreFunctionLibrary.h"
+#include "UnrealEngine.h"
+#include "Engine/Engine.h"
+#include "LatentActions.h"
+#include "Misc/AssertionMacros.h"
+#include "ARBlueprintLibrary.h"
+
+#include "GoogleARCoreAndroidHelper.h"
+#include "GoogleARCoreBaseLogCategory.h"
+#include "GoogleARCoreDevice.h"
+#include "GoogleARCoreXRTrackingSystem.h"
+
+namespace
+{
+	FGoogleARCoreXRTrackingSystem* GetARCoreHMD()
+	{
+		if (GEngine->XRSystem.IsValid() && (GEngine->XRSystem->GetSystemName() == FName("FGoogleARCoreXRTrackingSystem")))
+		{
+			return static_cast<FGoogleARCoreXRTrackingSystem*>(GEngine->XRSystem.Get());
+		}
+
+		return nullptr;
+	}
+
+	EGoogleARCoreInstallRequestResult ToAPKInstallStatus(EGoogleARCoreAPIStatus RequestStatus)
+	{
+		EGoogleARCoreInstallRequestResult OutRequestResult = EGoogleARCoreInstallRequestResult::FatalError;
+		switch (RequestStatus)
+		{
+			case EGoogleARCoreAPIStatus::AR_SUCCESS:
+				OutRequestResult = EGoogleARCoreInstallRequestResult::Installed;
+				break;
+			case EGoogleARCoreAPIStatus::AR_ERROR_FATAL:
+				OutRequestResult = EGoogleARCoreInstallRequestResult::FatalError;
+				break;
+			case EGoogleARCoreAPIStatus::AR_UNAVAILABLE_DEVICE_NOT_COMPATIBLE:
+				OutRequestResult = EGoogleARCoreInstallRequestResult::DeviceNotCompatible;
+				break;
+			case EGoogleARCoreAPIStatus::AR_UNAVAILABLE_USER_DECLINED_INSTALLATION:
+				OutRequestResult = EGoogleARCoreInstallRequestResult::UserDeclinedInstallation;
+				break;
+			default:
+				ensureMsgf(false, TEXT("Unexpected ARCore API Status: %d"), static_cast<int>(RequestStatus));
+				break;
+		}
+		return OutRequestResult;
+	}
+
+	const float DefaultLineTraceDistance = 100000; // 1000 meter
+}
+
+/************************************************************************/
+/*  UGoogleARCoreSessionFunctionLibrary | Lifecycle                     */
+/************************************************************************/
+struct FARCoreCheckAvailabilityAction : public FPendingLatentAction
+{
+public:
+	FName ExecutionFunction;
+	int32 OutputLink;
+	FWeakObjectPtr CallbackTarget;
+	EGoogleARCoreAvailability& OutAvailability;
+
+	FARCoreCheckAvailabilityAction(const FLatentActionInfo& InLatentInfo, EGoogleARCoreAvailability& InAvailability)
+	: FPendingLatentAction()
+	, ExecutionFunction(InLatentInfo.ExecutionFunction)
+	, OutputLink(InLatentInfo.Linkage)
+	, CallbackTarget(InLatentInfo.CallbackTarget)
+	, OutAvailability(InAvailability)
+	{}
+
+	virtual void UpdateOperation(FLatentResponse& Response) override
+	{
+		EGoogleARCoreAvailability ARCoreAvailability = FGoogleARCoreDevice::GetInstance()->CheckARCoreAPKAvailability();
+		if (ARCoreAvailability != EGoogleARCoreAvailability::UnknownChecking)
+		{
+			OutAvailability = ARCoreAvailability;
+			Response.FinishAndTriggerIf(true, ExecutionFunction, OutputLink, CallbackTarget);
+		}
+	}
+#if WITH_EDITOR
+	virtual FString GetDescription() const override
+	{
+		return FString::Printf(TEXT("Checking ARCore availability."));
+	}
+#endif
+};
+
+void UGoogleARCoreSessionFunctionLibrary::CheckARCoreAvailability(UObject* WorldContextObject, struct FLatentActionInfo LatentInfo, EGoogleARCoreAvailability& OutAvailability)
+{
+	if (UWorld* World = GEngine->GetWorldFromContextObject(WorldContextObject, EGetWorldErrorMode::LogAndReturnNull))
+	{
+		FLatentActionManager& LatentManager = World->GetLatentActionManager();
+		if (LatentManager.FindExistingAction<FARCoreCheckAvailabilityAction>(LatentInfo.CallbackTarget, LatentInfo.UUID) == nullptr)
+		{
+			FARCoreCheckAvailabilityAction* NewAction = new FARCoreCheckAvailabilityAction(LatentInfo, OutAvailability);
+			LatentManager.AddNewAction(LatentInfo.CallbackTarget, LatentInfo.UUID, NewAction);
+		}
+		else
+		{
+			UE_LOG(LogGoogleARCore, Warning, TEXT("CheckARCoreAvailability not excuted. The previous action hasn't finished yet."));
+		}
+	}
+}
+
+EGoogleARCoreAvailability UGoogleARCoreSessionFunctionLibrary::CheckARCoreAvailableStatus()
+{
+	return FGoogleARCoreDevice::GetInstance()->CheckARCoreAPKAvailability();
+}
+
+struct FARCoreAPKInstallAction : public FPendingLatentAction
+{
+public:
+	FName ExecutionFunction;
+	int32 OutputLink;
+	FWeakObjectPtr CallbackTarget;
+	EGoogleARCoreInstallRequestResult& OutRequestResult;
+	bool bInstallRequested;
+
+	FARCoreAPKInstallAction(const FLatentActionInfo& InLatentInfo, EGoogleARCoreInstallRequestResult& InRequestResult)
+	: FPendingLatentAction()
+	, ExecutionFunction(InLatentInfo.ExecutionFunction)
+	, OutputLink(InLatentInfo.Linkage)
+	, CallbackTarget(InLatentInfo.CallbackTarget)
+	, OutRequestResult(InRequestResult)
+	, bInstallRequested(false)
+	{}
+
+	virtual void UpdateOperation(FLatentResponse& Response) override
+	{
+		UE_LOG(LogTemp, Log, TEXT("IntallARCore UpdateOperation..."));
+		EGoogleARCoreInstallStatus InstallStatus = EGoogleARCoreInstallStatus::Installed;
+		EGoogleARCoreAPIStatus RequestStatus = FGoogleARCoreDevice::GetInstance()->RequestInstall(!bInstallRequested, InstallStatus);
+		UE_LOG(LogTemp, Log, TEXT("Requset ARCore Install(User Requested %d) Status: %d, Install Status: %d"), !bInstallRequested, (int)RequestStatus, (int)InstallStatus);
+		if (RequestStatus != EGoogleARCoreAPIStatus::AR_SUCCESS)
+		{
+			OutRequestResult = ToAPKInstallStatus(RequestStatus);
+			Response.FinishAndTriggerIf(true, ExecutionFunction, OutputLink, CallbackTarget);
+		}
+		else if (InstallStatus == EGoogleARCoreInstallStatus::Installed)
+		{
+			OutRequestResult = EGoogleARCoreInstallRequestResult::Installed;
+			Response.FinishAndTriggerIf(true, ExecutionFunction, OutputLink, CallbackTarget);
+		}
+		else
+		{
+			// InstallSttatus returns requested.
+			bInstallRequested = true;
+		}
+	}
+#if WITH_EDITOR
+	virtual FString GetDescription() const override
+	{
+		return FString::Printf(TEXT("Checking ARCore availability."));
+	}
+#endif
+};
+
+void UGoogleARCoreSessionFunctionLibrary::InstallARCoreService(UObject* WorldContextObject, struct FLatentActionInfo LatentInfo, EGoogleARCoreInstallRequestResult& OutInstallResult)
+{
+	if (UWorld* World = GEngine->GetWorldFromContextObject(WorldContextObject, EGetWorldErrorMode::LogAndReturnNull))
+	{
+		FLatentActionManager& LatentManager = World->GetLatentActionManager();
+		if (LatentManager.FindExistingAction<FARCoreAPKInstallAction>(LatentInfo.CallbackTarget, LatentInfo.UUID) == nullptr)
+		{
+			FARCoreAPKInstallAction* NewAction = new FARCoreAPKInstallAction(LatentInfo, OutInstallResult);
+			LatentManager.AddNewAction(LatentInfo.CallbackTarget, LatentInfo.UUID, NewAction);
+		}
+	}
+}
+
+EGoogleARCoreInstallStatus UGoogleARCoreSessionFunctionLibrary::RequestInstallARCoreAPK()
+{
+	EGoogleARCoreInstallStatus InstallStatus = EGoogleARCoreInstallStatus::Installed;
+	EGoogleARCoreAPIStatus RequestStatus = FGoogleARCoreDevice::GetInstance()->RequestInstall(true, InstallStatus);
+
+	return InstallStatus;
+}
+
+EGoogleARCoreInstallRequestResult UGoogleARCoreSessionFunctionLibrary::GetARCoreAPKInstallResult()
+{
+	EGoogleARCoreInstallStatus InstallStatus = EGoogleARCoreInstallStatus::Installed;
+	EGoogleARCoreAPIStatus RequestStatus = FGoogleARCoreDevice::GetInstance()->RequestInstall(false, InstallStatus);
+
+	return ToAPKInstallStatus(RequestStatus);
+}
+
+UGoogleARCoreEventManager* UGoogleARCoreSessionFunctionLibrary::GetARCoreEventManager()
+{
+	auto ARSystem = FGoogleARCoreDevice::GetInstance()->GetARSystem();
+	if (ARSystem.IsValid())
+	{
+		return static_cast<FGoogleARCoreXRTrackingSystem*>(ARSystem->GetXRTrackingSystem())->GetEventManager();
+	}
+	else
+	{
+		return nullptr;
+	}
+	
+}
+
+struct FARCoreStartSessionAction : public FPendingLatentAction
+{
+public:
+	FName ExecutionFunction;
+	int32 OutputLink;
+	FWeakObjectPtr CallbackTarget;
+
+	FARCoreStartSessionAction(const FLatentActionInfo& InLatentInfo)
+		: FPendingLatentAction()
+		, ExecutionFunction(InLatentInfo.ExecutionFunction)
+		, OutputLink(InLatentInfo.Linkage)
+		, CallbackTarget(InLatentInfo.CallbackTarget)
+	{}
+
+	virtual void UpdateOperation(FLatentResponse& Response) override
+	{
+		bool bSessionStartedFinished = FGoogleARCoreDevice::GetInstance()->GetStartSessionRequestFinished();
+		Response.FinishAndTriggerIf(bSessionStartedFinished, ExecutionFunction, OutputLink, CallbackTarget);
+	}
+#if WITH_EDITOR
+	virtual FString GetDescription() const override
+	{
+		return FString::Printf(TEXT("Starting ARCore tracking session"));
+	}
+#endif
+};
+
+void UGoogleARCoreSessionFunctionLibrary::StartARCoreSession(UObject* WorldContextObject, struct FLatentActionInfo LatentInfo, UGoogleARCoreSessionConfig* Configuration)
+{
+	if (UWorld* World = GEngine->GetWorldFromContextObject(WorldContextObject, EGetWorldErrorMode::LogAndReturnNull))
+	{
+		FLatentActionManager& LatentManager = World->GetLatentActionManager();
+		if (LatentManager.FindExistingAction<FARCoreStartSessionAction>(LatentInfo.CallbackTarget, LatentInfo.UUID) == nullptr)
+		{
+			UARBlueprintLibrary::StartARSession(static_cast<UARSessionConfig*>(Configuration));
+			FARCoreStartSessionAction* NewAction = new FARCoreStartSessionAction(LatentInfo);
+			LatentManager.AddNewAction(LatentInfo.CallbackTarget, LatentInfo.UUID, NewAction);
+		}
+	}
+}
+
+bool UGoogleARCoreSessionFunctionLibrary::SetARCoreCameraConfig(FGoogleARCoreCameraConfig TargetCameraConfig)
+{
+	return FGoogleARCoreDevice::GetInstance()->SetARCameraConfig(TargetCameraConfig);
+}
+
+bool UGoogleARCoreSessionFunctionLibrary::GetARCoreCameraConfig(FGoogleARCoreCameraConfig& OutCurrentCameraConfig)
+{
+	return FGoogleARCoreDevice::GetInstance()->GetARCameraConfig(OutCurrentCameraConfig);
+}
+
+/************************************************************************/
+/*  UGoogleARCoreSessionFunctionLibrary | PassthroughCamera             */
+/************************************************************************/
+bool UGoogleARCoreSessionFunctionLibrary::IsPassthroughCameraRenderingEnabled()
+{
+	FGoogleARCoreXRTrackingSystem* ARCoreHMD = GetARCoreHMD();
+	if (ARCoreHMD)
+	{
+		return ARCoreHMD->GetColorCameraRenderingEnabled();
+	}
+	else
+	{
+		UE_LOG(LogGoogleARCore, Error, TEXT("Failed to find GoogleARCoreXRTrackingSystem: GoogleARCoreXRTrackingSystem is not available."));
+	}
+	return false;
+}
+
+void UGoogleARCoreSessionFunctionLibrary::SetPassthroughCameraRenderingEnabled(bool bEnable)
+{
+	FGoogleARCoreXRTrackingSystem* ARCoreHMD = GetARCoreHMD();
+	if (ARCoreHMD)
+	{
+		ARCoreHMD->EnableColorCameraRendering(bEnable);
+	}
+	else
+	{
+		UE_LOG(LogGoogleARCore, Error, TEXT("Failed to config ARCoreXRCamera: GoogleARCoreXRTrackingSystem is not available."));
+	}
+}
+
+void UGoogleARCoreSessionFunctionLibrary::GetPassthroughCameraImageUV(const TArray<float>& InUV, TArray<float>& OutUV)
+{
+	FGoogleARCoreDevice::GetInstance()->GetPassthroughCameraImageUVs(InUV, OutUV);
+}
+
+void UGoogleARCoreSessionFunctionLibrary::GetAllPlanes(TArray<UARPlaneGeometry*>& OutPlaneList)
+{
+	FGoogleARCoreDevice::GetInstance()->GetAllTrackables<UARPlaneGeometry>(OutPlaneList);
+}
+
+void UGoogleARCoreSessionFunctionLibrary::GetAllTrackablePoints(TArray<UARTrackedPoint*>& OutTrackablePointList)
+{
+	FGoogleARCoreDevice::GetInstance()->GetAllTrackables<UARTrackedPoint>(OutTrackablePointList);
+}
+
+template< class T >
+void UGoogleARCoreSessionFunctionLibrary::GetAllTrackable(TArray<T*>& OutTrackableList)
+{
+	FGoogleARCoreDevice::GetInstance()->GetAllTrackables<T>(OutTrackableList);
+}
+
+template void UGoogleARCoreSessionFunctionLibrary::GetAllTrackable<UARTrackedGeometry>(TArray<UARTrackedGeometry*>& OutTrackableList);
+template void UGoogleARCoreSessionFunctionLibrary::GetAllTrackable<UARPlaneGeometry>(TArray<UARPlaneGeometry*>& OutTrackableList);
+template void UGoogleARCoreSessionFunctionLibrary::GetAllTrackable<UARTrackedPoint>(TArray<UARTrackedPoint*>& OutTrackableList);
+
+/************************************************************************/
+/*  UGoogleARCoreFrameFunctionLibrary                                   */
+/************************************************************************/
+EGoogleARCoreTrackingState UGoogleARCoreFrameFunctionLibrary::GetTrackingState()
+{
+	return FGoogleARCoreDevice::GetInstance()->GetTrackingState();
+}
+
+void UGoogleARCoreFrameFunctionLibrary::GetPose(FTransform& LastePose)
+{
+	LastePose = FGoogleARCoreDevice::GetInstance()->GetLatestPose();
+}
+
+bool UGoogleARCoreFrameFunctionLibrary::ARCoreLineTrace(UObject* WorldContextObject, const FVector2D& ScreenPosition, TSet<EGoogleARCoreLineTraceChannel> TraceChannels, TArray<FARTraceResult>& OutHitResults)
+{
+	EGoogleARCoreLineTraceChannel TraceChannelValue = EGoogleARCoreLineTraceChannel::None;
+	for (EGoogleARCoreLineTraceChannel Channel : TraceChannels)
+	{
+		TraceChannelValue = TraceChannelValue | Channel;
+	}
+
+	FGoogleARCoreDevice::GetInstance()->ARLineTrace(ScreenPosition, TraceChannelValue, OutHitResults);
+	return OutHitResults.Num() > 0;
+}
+
+bool UGoogleARCoreFrameFunctionLibrary::ARCoreLineTraceRay(UObject* WorldContextObject, const FVector& Start, const FVector& End, TSet<EGoogleARCoreLineTraceChannel> TraceChannels, TArray<FARTraceResult>& OutHitResults)
+{
+	EGoogleARCoreLineTraceChannel TraceChannelValue = EGoogleARCoreLineTraceChannel::None;
+	for (EGoogleARCoreLineTraceChannel Channel : TraceChannels)
+	{
+		TraceChannelValue = TraceChannelValue | Channel;
+	}
+
+	FGoogleARCoreDevice::GetInstance()->ARLineTrace(Start, End, TraceChannelValue, OutHitResults);
+	return OutHitResults.Num() > 0;
+}
+
+void UGoogleARCoreFrameFunctionLibrary::GetUpdatedARPins(TArray<UARPin*>& OutAnchorList)
+{
+	FGoogleARCoreDevice::GetInstance()->GetUpdatedARPins(OutAnchorList);
+}
+
+void UGoogleARCoreFrameFunctionLibrary::GetUpdatedPlanes(TArray<UARPlaneGeometry*>& OutPlaneList)
+{
+	FGoogleARCoreDevice::GetInstance()->GetUpdatedTrackables<UARPlaneGeometry>(OutPlaneList);
+}
+
+void UGoogleARCoreFrameFunctionLibrary::GetUpdatedTrackablePoints(TArray<UARTrackedPoint*>& OutTrackablePointList)
+{
+	FGoogleARCoreDevice::GetInstance()->GetUpdatedTrackables<UARTrackedPoint>(OutTrackablePointList);
+}
+
+void UGoogleARCoreFrameFunctionLibrary::GetLightEstimation(FGoogleARCoreLightEstimate& LightEstimation)
+{
+	LightEstimation = FGoogleARCoreDevice::GetInstance()->GetLatestLightEstimate();
+}
+
+EGoogleARCoreFunctionStatus UGoogleARCoreFrameFunctionLibrary::GetPointCloud(UGoogleARCorePointCloud*& OutLatestPointCloud)
+{
+	return FGoogleARCoreDevice::GetInstance()->GetLatestPointCloud(OutLatestPointCloud);
+}
+
+EGoogleARCoreFunctionStatus UGoogleARCoreFrameFunctionLibrary::AcquirePointCloud(UGoogleARCorePointCloud*& OutLatestPointCloud)
+{
+	return FGoogleARCoreDevice::GetInstance()->AcquireLatestPointCloud(OutLatestPointCloud);
+}
+
+#if PLATFORM_ANDROID
+EGoogleARCoreFunctionStatus UGoogleARCoreFrameFunctionLibrary::GetCameraMetadata(const ACameraMetadata*& OutCameraMetadata)
+{
+	return FGoogleARCoreDevice::GetInstance()->GetLatestCameraMetadata(OutCameraMetadata);
+}
+#endif
+
+template< class T >
+void UGoogleARCoreFrameFunctionLibrary::GetUpdatedTrackable(TArray<T*>& OutTrackableList)
+{
+	FGoogleARCoreDevice::GetInstance()->GetUpdatedTrackables<T>(OutTrackableList);
+}
+
+template void UGoogleARCoreFrameFunctionLibrary::GetUpdatedTrackable<UARTrackedGeometry>(TArray<UARTrackedGeometry*>& OutTrackableList);
+template void UGoogleARCoreFrameFunctionLibrary::GetUpdatedTrackable<UARPlaneGeometry>(TArray<UARPlaneGeometry*>& OutTrackableList);
+template void UGoogleARCoreFrameFunctionLibrary::GetUpdatedTrackable<UARTrackedPoint>(TArray<UARTrackedPoint*>& OutTrackableList);
+
+
+EGoogleARCoreFunctionStatus UGoogleARCoreFrameFunctionLibrary::AcquireCameraImage(UGoogleARCoreCameraImage *&OutLatestCameraImage)
+{
+	return FGoogleARCoreDevice::GetInstance()->AcquireCameraImage(OutLatestCameraImage);
+}
+
+EGoogleARCoreFunctionStatus UGoogleARCoreFrameFunctionLibrary::GetCameraImageIntrinsics(UGoogleARCoreCameraIntrinsics *&OutCameraIntrinsics)
+{
+	return FGoogleARCoreDevice::GetInstance()->GetCameraImageIntrinsics(OutCameraIntrinsics);
+}
+
+EGoogleARCoreFunctionStatus UGoogleARCoreFrameFunctionLibrary::GetCameraTextureIntrinsics(UGoogleARCoreCameraIntrinsics *&OutCameraIntrinsics)
+{
+	return FGoogleARCoreDevice::GetInstance()->GetCameraTextureIntrinsics(OutCameraIntrinsics);
+}
+
diff --git a/Engine/Plugins/Experimental/AR/Google/GoogleARCore/Source/GoogleARCoreBase/Private/GoogleARCorePlaneRendererComponent.cpp b/Engine/Plugins/Experimental/AR/Google/GoogleARCore/Source/GoogleARCoreBase/Private/GoogleARCorePlaneRendererComponent.cpp
new file mode 100644
--- /dev/null
+++ b/Engine/Plugins/Experimental/AR/Google/GoogleARCore/Source/GoogleARCoreBase/Private/GoogleARCorePlaneRendererComponent.cpp
@@ -0,0 +1,78 @@
+// Copyright 2017 Google Inc.
+
+#include "GoogleARCorePlaneRendererComponent.h"
+#include "ARBlueprintLibrary.h"
+#include "ARSystem.h"
+#include "ARTrackable.h"
+#include "DrawDebugHelpers.h"
+#include "Templates/Casts.h"
+
+UGoogleARCorePlaneRendererComponent::UGoogleARCorePlaneRendererComponent()
+	: bRenderPlane(true)
+	, bRenderBoundaryPolygon(true)
+	, PlaneColor(FColor::Green)
+	, BoundaryPolygonColor(FColor::Blue)
+	, BoundaryPolygonThickness(0.5f)
+{
+	PrimaryComponentTick.bCanEverTick = true;
+
+	PlaneIndices.AddUninitialized(6);
+	PlaneIndices[0] = 0; PlaneIndices[1] = 1; PlaneIndices[2] = 2;
+	PlaneIndices[3] = 0; PlaneIndices[4] = 2; PlaneIndices[5] = 3;
+}
+
+void UGoogleARCorePlaneRendererComponent::TickComponent(float DeltaTime, enum ELevelTick TickType, FActorComponentTickFunction * ThisTickFunction)
+{
+	DrawPlanes();
+}
+
+void UGoogleARCorePlaneRendererComponent::DrawPlanes()
+{
+	UWorld* World = GetWorld();
+	if (UARBlueprintLibrary::GetTrackingQuality() == EARTrackingQuality::OrientationAndPosition)
+	{
+		TArray<UARTrackedGeometry*> PlaneList;
+		PlaneList = UARBlueprintLibrary::GetAllGeometries();
+		for (UARTrackedGeometry* TrackedGeometry : PlaneList)
+		{
+			if (!TrackedGeometry->IsA(UARPlaneGeometry::StaticClass()))
+			{
+				continue;
+			}
+			
+			if (TrackedGeometry->GetTrackingState() != EARTrackingState::Tracking)
+			{
+				continue;
+			}
+			
+			UARPlaneGeometry* Plane = Cast<UARPlaneGeometry>(TrackedGeometry);
+			if (bRenderPlane)
+			{
+				FTransform BoundingBoxTransform = Plane->GetLocalToWorldTransform();
+				FVector BoundingBoxLocation = BoundingBoxTransform.GetLocation();
+				FVector BoundingBoxSize = Plane->GetExtent();
+				
+				PlaneVertices.Empty();
+				PlaneVertices.Add(BoundingBoxTransform.TransformPosition(FVector(-BoundingBoxSize.X, -BoundingBoxSize.Y, 0)));
+				PlaneVertices.Add(BoundingBoxTransform.TransformPosition(FVector(-BoundingBoxSize.X, BoundingBoxSize.Y, 0)));
+				PlaneVertices.Add(BoundingBoxTransform.TransformPosition(FVector(BoundingBoxSize.X, BoundingBoxSize.Y, 0)));
+				PlaneVertices.Add(BoundingBoxTransform.TransformPosition(FVector(BoundingBoxSize.X, -BoundingBoxSize.Y, 0)));
+				// plane quad
+				DrawDebugMesh(World, PlaneVertices, PlaneIndices, PlaneColor);
+			}
+			
+			if (bRenderBoundaryPolygon)
+			{
+				const TArray<FVector>& BoundaryPolygonData = Plane->GetBoundaryPolygonInLocalSpace();
+				FTransform PlaneCenter = Plane->GetLocalToWorldTransform();
+				for (int i = 0; i < BoundaryPolygonData.Num(); i++)
+				{
+					FVector Start = PlaneCenter.TransformPosition(BoundaryPolygonData[i]);
+					FVector End = PlaneCenter.TransformPosition(BoundaryPolygonData[(i + 1) % BoundaryPolygonData.Num()]);
+					DrawDebugLine(World, Start, End, BoundaryPolygonColor, false, -1.f, 0, BoundaryPolygonThickness);
+				}
+			}
+
+		}
+	}
+}
diff --git a/Engine/Plugins/Experimental/AR/Google/GoogleARCore/Source/GoogleARCoreBase/Private/GoogleARCorePointCloudRendererComponent.cpp b/Engine/Plugins/Experimental/AR/Google/GoogleARCore/Source/GoogleARCoreBase/Private/GoogleARCorePointCloudRendererComponent.cpp
new file mode 100644
--- /dev/null
+++ b/Engine/Plugins/Experimental/AR/Google/GoogleARCore/Source/GoogleARCoreBase/Private/GoogleARCorePointCloudRendererComponent.cpp
@@ -0,0 +1,39 @@
+// Copyright 2017 Google Inc.
+
+#include "GoogleARCorePointCloudRendererComponent.h"
+#include "DrawDebugHelpers.h"
+#include "GoogleARCoreTypes.h"
+#include "GoogleARCoreFunctionLibrary.h"
+
+UGoogleARCorePointCloudRendererComponent::UGoogleARCorePointCloudRendererComponent()
+	: PointColor(FColor::Red)
+	, PointSize(0.1f)
+{
+	PreviousPointCloudTimestamp = 0.0;
+	PrimaryComponentTick.bCanEverTick = true;
+}
+
+void UGoogleARCorePointCloudRendererComponent::TickComponent(float DeltaTime, enum ELevelTick TickType, FActorComponentTickFunction * ThisTickFunction)
+{
+	DrawPointCloud();
+}
+
+void UGoogleARCorePointCloudRendererComponent::DrawPointCloud()
+{
+	UWorld* World = GetWorld();
+	if (UGoogleARCoreFrameFunctionLibrary::GetTrackingState() == EGoogleARCoreTrackingState::Tracking)
+	{
+		UGoogleARCorePointCloud* LatestPointCloud = nullptr;
+		EGoogleARCoreFunctionStatus Status = UGoogleARCoreFrameFunctionLibrary::GetPointCloud(LatestPointCloud);
+		if (Status == EGoogleARCoreFunctionStatus::Success && LatestPointCloud != nullptr && LatestPointCloud->GetPointNum() > 0)
+		{
+			for (int i = 0; i < LatestPointCloud->GetPointNum(); i++)
+			{
+				FVector PointPosition = FVector::ZeroVector;
+				float PointConfidence = 0;
+				LatestPointCloud->GetPoint(i, PointPosition, PointConfidence);
+				DrawDebugPoint(World, PointPosition, PointSize, PointColor, false);
+			}
+		}
+	}
+}
diff --git a/Engine/Plugins/Experimental/AR/Google/GoogleARCore/Source/GoogleARCoreBase/Private/GoogleARCoreXRCamera.cpp b/Engine/Plugins/Experimental/AR/Google/GoogleARCore/Source/GoogleARCoreBase/Private/GoogleARCoreXRCamera.cpp
new file mode 100644
--- /dev/null
+++ b/Engine/Plugins/Experimental/AR/Google/GoogleARCore/Source/GoogleARCoreBase/Private/GoogleARCoreXRCamera.cpp
@@ -0,0 +1,78 @@
+// Copyright 1998-2018 Epic Games, Inc. All Rights Reserved.
+
+#include "GoogleARCoreXRCamera.h"
+#include "GoogleARCoreXRTrackingSystem.h"
+#include "SceneView.h"
+#include "GoogleARCorePassthroughCameraRenderer.h"
+#include "GoogleARCoreAndroidHelper.h"
+
+FGoogleARCoreXRCamera::FGoogleARCoreXRCamera(const FAutoRegister& AutoRegister, FGoogleARCoreXRTrackingSystem& InARCoreSystem, int32 InDeviceID)
+	: FDefaultXRCamera(AutoRegister, &InARCoreSystem, InDeviceID)
+	, GoogleARCoreTrackingSystem(InARCoreSystem)
+	, bMatchDeviceCameraFOV(false)
+	, bEnablePassthroughCameraRendering_RT(false)
+{
+	PassthroughRenderer = new FGoogleARCorePassthroughCameraRenderer();
+}
+
+void FGoogleARCoreXRCamera::SetupView(FSceneViewFamily& InViewFamily, FSceneView& InView)
+{
+	TrackingSystem->GetCurrentPose(DeviceId, InView.BaseHmdOrientation, InView.BaseHmdLocation);
+}
+
+void FGoogleARCoreXRCamera::SetupViewProjectionMatrix(FSceneViewProjectionData& InOutProjectionData)
+{
+	if (GoogleARCoreTrackingSystem.ARCoreDeviceInstance->GetIsARCoreSessionRunning() && bMatchDeviceCameraFOV)
+	{
+		FIntRect ViewRect = InOutProjectionData.GetViewRect();
+		InOutProjectionData.ProjectionMatrix = GoogleARCoreTrackingSystem.ARCoreDeviceInstance->GetPassthroughCameraProjectionMatrix(ViewRect.Size());
+	}
+}
+
+void FGoogleARCoreXRCamera::BeginRenderViewFamily(FSceneViewFamily& InViewFamily)
+{
+	PassthroughRenderer->InitializeOverlayMaterial();
+}
+
+void FGoogleARCoreXRCamera::PreRenderView_RenderThread(FRHICommandListImmediate& RHICmdList, FSceneView& InView)
+{
+}
+
+void FGoogleARCoreXRCamera::PreRenderViewFamily_RenderThread(FRHICommandListImmediate& RHICmdList, FSceneViewFamily& InViewFamily)
+{
+	FGoogleARCoreXRTrackingSystem& TS = GoogleARCoreTrackingSystem;
+
+	if (TS.ARCoreDeviceInstance->GetIsARCoreSessionRunning() && bEnablePassthroughCameraRendering_RT)
+	{
+		PassthroughRenderer->InitializeRenderer_RenderThread(TS.ARCoreDeviceInstance->GetPassthroughCameraTexture());
+	}
+}
+
+void FGoogleARCoreXRCamera::PostRenderBasePass_RenderThread(FRHICommandListImmediate& RHICmdList, FSceneView& InView)
+{
+	if (GoogleARCoreTrackingSystem.ARCoreDeviceInstance->GetIsARCoreSessionRunning() && bEnablePassthroughCameraRendering_RT)
+	{
+		TArray<float> TransformedUVs;
+		GoogleARCoreTrackingSystem.ARCoreDeviceInstance->GetPassthroughCameraImageUVs(PassthroughRenderer->OverlayQuadUVs, TransformedUVs);
+		PassthroughRenderer->UpdateOverlayUVCoordinate_RenderThread(TransformedUVs, FGoogleARCoreAndroidHelper::GetDisplayRotation());
+		PassthroughRenderer->RenderVideoOverlay_RenderThread(RHICmdList, InView);
+	}
+}
+
+bool FGoogleARCoreXRCamera::IsActiveThisFrame(class FViewport* InViewport) const
+{
+	return GoogleARCoreTrackingSystem.IsHeadTrackingAllowed();
+}
+
+void FGoogleARCoreXRCamera::ConfigXRCamera(bool bInMatchDeviceCameraFOV, bool bInEnablePassthroughCameraRendering)
+{
+	bMatchDeviceCameraFOV = bInMatchDeviceCameraFOV;
+	ENQUEUE_UNIQUE_RENDER_COMMAND_TWOPARAMETER(
+		ConfigXRCamera,
+		FGoogleARCoreXRCamera*, ARCoreXRCamera, this,
+		bool, bInEnablePassthroughCameraRendering, bInEnablePassthroughCameraRendering,
+		{
+			ARCoreXRCamera->bEnablePassthroughCameraRendering_RT = bInEnablePassthroughCameraRendering;
+		}
+	);
+}
diff --git a/Engine/Plugins/Experimental/AR/Google/GoogleARCore/Source/GoogleARCoreBase/Private/GoogleARCoreXRCamera.h b/Engine/Plugins/Experimental/AR/Google/GoogleARCore/Source/GoogleARCoreBase/Private/GoogleARCoreXRCamera.h
new file mode 100644
--- /dev/null
+++ b/Engine/Plugins/Experimental/AR/Google/GoogleARCore/Source/GoogleARCoreBase/Private/GoogleARCoreXRCamera.h
@@ -0,0 +1,33 @@
+// Copyright 1998-2018 Epic Games, Inc. All Rights Reserved.
+
+#pragma once
+
+#include "DefaultXRCamera.h"
+
+class FGoogleARCoreXRTrackingSystem;
+class FGoogleARCorePassthroughCameraRenderer;
+
+class FGoogleARCoreXRCamera : public FDefaultXRCamera
+{
+
+public:
+	FGoogleARCoreXRCamera(const FAutoRegister&, FGoogleARCoreXRTrackingSystem& InARCoreSystem, int32 InDeviceID);
+
+	//~ FDefaultXRCamera
+	virtual void SetupView(FSceneViewFamily& InViewFamily, FSceneView& InView) override;
+	virtual void SetupViewProjectionMatrix(FSceneViewProjectionData& InOutProjectionData) override;
+	virtual void BeginRenderViewFamily(FSceneViewFamily& InViewFamily) override;
+	virtual void PreRenderView_RenderThread(FRHICommandListImmediate& RHICmdList, FSceneView& InView) override;
+	virtual void PreRenderViewFamily_RenderThread(FRHICommandListImmediate& RHICmdList, FSceneViewFamily& InViewFamily) override;
+	virtual void PostRenderBasePass_RenderThread(FRHICommandListImmediate& RHICmdList, FSceneView& InView) override;
+	virtual bool IsActiveThisFrame(class FViewport* InViewport) const override;
+	//~ FDefaultXRCamera
+
+	void ConfigXRCamera(bool bInMatchDeviceCameraFOV, bool bInEnablePassthroughCameraRendering);
+private:
+	FGoogleARCoreXRTrackingSystem& GoogleARCoreTrackingSystem;
+	FGoogleARCorePassthroughCameraRenderer* PassthroughRenderer;
+
+	bool bMatchDeviceCameraFOV;
+	bool bEnablePassthroughCameraRendering_RT;
+};
diff --git a/Engine/Plugins/Experimental/AR/Google/GoogleARCore/Source/GoogleARCoreBase/Private/Tests/GoogleARCoreDummyTest.cpp b/Engine/Plugins/Experimental/AR/Google/GoogleARCore/Source/GoogleARCoreBase/Private/Tests/GoogleARCoreDummyTest.cpp
new file mode 100644
--- /dev/null
+++ b/Engine/Plugins/Experimental/AR/Google/GoogleARCore/Source/GoogleARCoreBase/Private/Tests/GoogleARCoreDummyTest.cpp
@@ -0,0 +1,8 @@
+#include "Misc/AutomationTest.h"
+
+IMPLEMENT_SIMPLE_AUTOMATION_TEST(FGoogleARCoreDummyTest, "GoogleARCore.DummyTest", EAutomationTestFlags::ApplicationContextMask | EAutomationTestFlags::EngineFilter)
+
+bool FGoogleARCoreDummyTest::RunTest(const FString& Parameters)
+{
+          return true;
+}
diff --git a/Engine/Plugins/Experimental/AR/Google/GoogleARCore/Source/GoogleARCoreBase/Public/GoogleARCoreAugmentedImage.h b/Engine/Plugins/Experimental/AR/Google/GoogleARCore/Source/GoogleARCoreBase/Public/GoogleARCoreAugmentedImage.h
new file mode 100644
--- /dev/null
+++ b/Engine/Plugins/Experimental/AR/Google/GoogleARCore/Source/GoogleARCoreBase/Public/GoogleARCoreAugmentedImage.h
@@ -0,0 +1,84 @@
+// Copyright 2017 Google Inc.
+
+#pragma once
+
+#include "GoogleARCoreAugmentedImageDatabase.h"
+
+#include "ARTrackable.h"
+
+#include "GoogleARCoreAugmentedImage.generated.h"
+
+/**
+ * An object representing an augmented image currently in the scene.
+ */
+UCLASS(BlueprintType)
+class GOOGLEARCOREBASE_API UGoogleARCoreAugmentedImage : public UARTrackedGeometry
+{
+	GENERATED_BODY()
+
+public:
+
+	/**
+	 * Get the center in local space of the augmented image.
+	 *
+	 * @return The center in the augmented image's local coordinate
+	 *         space.
+	 */
+	UFUNCTION(BlueprintPure, Category = "GoogleARCore|AugmentedImage", meta = (Keywords = "googlear arcore augmentedimage"))
+	FVector GetCenter() const;
+
+	/**
+	 * Get the size in local space of the augmented image.
+	 *
+	 * @return The size in the augmented image's local coordinate
+	 *         space.
+	 */
+	UFUNCTION(BlueprintPure, Category = "GoogleARCore|AugmentedImage", meta = (Keywords = "googlear arcore augmentedimage"))
+	FVector GetExtent() const;
+
+	/**
+	 * Get the index of an augmented image in the augmented image
+	 * database.
+	 *
+	 * @return An integer indicating the index into the Entries array
+	 *         in the UGoogleARCoreAugmentedImageDatabase.
+	 */
+	UFUNCTION(BlueprintPure, Category = "GoogleARCore|AugmentedImage", meta = (Keywords = "googlear arcore augmentedimage"))
+	int32 GetImageIndex() const;
+
+	/**
+	 * Get the name of an augmented image in the augmented image
+	 * database.
+	 *
+	 * @return The name of the image in the image database.
+	 */
+	UFUNCTION(BlueprintPure, Category = "GoogleARCore|AugmentedImage", meta = (Keywords = "googlear arcore augmentedimage"))
+	FString GetImageName() const;
+
+	virtual void DebugDraw(
+		UWorld* World, const FLinearColor& OutlineColor,
+		float OutlineThickness, float PersistForSeconds = 0.0f) const override;
+
+	void UpdateTrackedGeometry(
+		const TSharedRef<FARSupportInterface , ESPMode::ThreadSafe>& InTrackingSystem,
+		uint32 FrameNumber, double Timestamp, const FTransform& InLocalToTrackingTransform,
+		const FTransform& InAlignmentTransform,
+		const FVector &InCenter, const FVector &InExtent,
+		int32 InImageIndex, const FString& ImageName);
+
+private:
+
+	UPROPERTY()
+	FVector Center;
+
+	UPROPERTY()
+	FVector Extent;
+
+	UPROPERTY()
+	int32 ImageIndex;
+
+	UPROPERTY()
+	FString ImageName;
+
+};
+
diff --git a/Engine/Plugins/Experimental/AR/Google/GoogleARCore/Source/GoogleARCoreBase/Public/GoogleARCoreFunctionLibrary.h b/Engine/Plugins/Experimental/AR/Google/GoogleARCore/Source/GoogleARCoreBase/Public/GoogleARCoreFunctionLibrary.h
new file mode 100644
--- /dev/null
+++ b/Engine/Plugins/Experimental/AR/Google/GoogleARCore/Source/GoogleARCoreBase/Public/GoogleARCoreFunctionLibrary.h
@@ -0,0 +1,296 @@
+// Copyright 2017 Google Inc.
+
+#pragma once
+
+#include "CoreMinimal.h"
+#include "Engine/LatentActionManager.h"
+#include "Kismet/BlueprintFunctionLibrary.h"
+
+#include "GoogleARCoreTypes.h"
+#include "GoogleARCoreSessionConfig.h"
+#include "GoogleARCoreFunctionLibrary.generated.h"
+
+/** A function library that provides static/Blueprint functions associated with GoogleARCore session.*/
+UCLASS()
+class GOOGLEARCOREBASE_API UGoogleARCoreSessionFunctionLibrary : public UBlueprintFunctionLibrary
+{
+	GENERATED_BODY()
+
+public:
+	//-----------------Lifecycle---------------------
+
+	/**
+	 * A Latent Action to check the availability of ARCore on this device.
+	 * This may initiate a query with a remote service to determine if the device is supported by ARCore. The Latent Action will complete when the check is finished.
+	 *
+	 * @param OutAvailability	The availability result as a EGoogleARCoreAvailability.
+	 */
+	UFUNCTION(BlueprintCallable, Category = "GoogleARCore|Availability", meta = (Latent, LatentInfo = "LatentInfo", WorldContext = "WorldContextObject", Keywords = "googlear arcore availability"))
+	static void CheckARCoreAvailability(UObject* WorldContextObject, struct FLatentActionInfo LatentInfo, EGoogleARCoreAvailability& OutAvailability);
+
+	/**
+	 * A Latent Action to initiates installation of ARCore if required.
+	 * This function may cause your application be paused if installing ARCore is required.
+	 *
+	 * @param OutInstallResult	The install request result as a EGoogleARCoreInstallRequestResult.
+	 */
+	UFUNCTION(BlueprintCallable, Category = "GoogleARCore|Availability", meta = (Latent, LatentInfo = "LatentInfo", WorldContext = "WorldContextObject", Keywords = "googlear arcore availability"))
+	static void InstallARCoreService(UObject* WorldContextObject, struct FLatentActionInfo LatentInfo, EGoogleARCoreInstallRequestResult& OutInstallResult);
+
+	/**
+	 * A polling function to check the ARCore availability in C++.
+	 * This may initiate a query with a remote service to determine if the device is supported by ARCore, so this function will EGoogleARCoreAvailability::UnknownChecking.
+	 *
+	 * @return	The availability result as a EGoogleARCoreAvailability.
+	 */
+	static EGoogleARCoreAvailability CheckARCoreAvailableStatus();
+
+	/**
+	 * Initiates installation of ARCore if required.
+	 * This function will return immediately and may pause your application if installing ARCore is required.
+	 *
+	 * @return EGoogleARCoreInstallStatus::Requrested if it started a install request.
+	 */
+	static EGoogleARCoreInstallStatus RequestInstallARCoreAPK();
+
+	/**
+	 * A polling function to check the ARCore install request result in C++.
+	 * After you call RequestInstallARCoreAPK() and it returns EGoogleARCoreInstallStatus::Requrested. You can call this function to check the install requst result.
+	 *
+	 * @return The install request result as a EGoogleARCoreInstallRequestResult.
+	 */
+	static EGoogleARCoreInstallRequestResult GetARCoreAPKInstallResult();
+
+	/**
+	 * Get the UGoogleARCoreEventManager to bind BP events or c++ delegate in GoogleARCore plugins.
+	 */
+	UFUNCTION(BlueprintPure, Category = "GoogleARCore|Session", meta = (Keywords = "googlear arcore event manager"))
+	static UGoogleARCoreEventManager* GetARCoreEventManager();
+
+	/**
+	 * Starts a new ARCore tracking session GoogleARCore specific configuration.
+	 * If the session already started and the config isn't the same, it will stop the previous session and start a new session with the new config.
+	 * Note that this is a latent action, you can query the session start result by querying GetARCoreSessionStatus() after the latent action finished.
+	 *
+	 * @param Configuration				The ARCoreSession configuration to start the session.
+	 */
+	UFUNCTION(BlueprintCallable, Category = "GoogleARCore|Session", meta = (Latent, LatentInfo = "LatentInfo", WorldContext = "WorldContextObject", Keywords = "googlear arcore session start config"))
+	static void StartARCoreSession(UObject* WorldContextObject, struct FLatentActionInfo LatentInfo, UGoogleARCoreSessionConfig* Configuration);
+
+	/**
+	 * Configure the ARCoreSession with the desired camera configuration. The TargetCameraConfig must be 
+	 * from a list returned by UGoogleARCoreEventManager::OnCameraConfig delegate.
+	 *
+	 * This function should be called when UGoogleARCoreEventManager::OnCameraConfig delegate got triggered.
+	 */
+	UFUNCTION(BlueprintCallable, Category = "GoogleARCore|Session", meta = (Keywords = "googlear arcore camera config"))
+	static bool SetARCoreCameraConfig(FGoogleARCoreCameraConfig TargetCameraConfig);
+
+	/** 
+	 * Get the FGoogleARCoreCameraConfig that the current ARCore session is using. 
+	 * 
+	 * @param OutCurrentCameraConfig   The FGoogleARCoreCameraConfig that the current ARCore session is using.
+	 * @return  True if there is a valid ARCore session and the current camera config is returned.
+	 *          False if ARCore session hasn't been started or it is already stopped.
+	 */
+	UFUNCTION(BlueprintPure, Category = "GoogleARCore|Session", meta = (Keywords = "googlear arcore camera config"))
+	static bool GetARCoreCameraConfig(FGoogleARCoreCameraConfig& OutCurrentCameraConfig);
+
+	//-----------------PassthroughCamera---------------------
+	/**
+	 * Returns the state of the passthrough camera rendering in GoogleARCore ARSystem.
+	 *
+	 * @return	True if the passthrough camera rendering is enabled.
+	 */
+	UFUNCTION(BlueprintPure, Category = "GoogleARCore|PassthroughCamera", meta = (Keywords = "googlear arcore passthrough camera"))
+	static bool IsPassthroughCameraRenderingEnabled();
+
+	/**
+	 * Enables/Disables the passthrough camera rendering in GoogleARCore ARSystem.
+	 * Note that when passthrough camera rendering is enabled, the camera FOV will be forced
+	 * to match FOV of the physical camera on the device.
+	 */
+	UFUNCTION(BlueprintCallable, Category = "GoogleARCore|PassthroughCamera", meta = (Keywords = "googlear arcore passthrough camera"))
+	static void SetPassthroughCameraRenderingEnabled(bool bEnable);
+
+	/**
+	 * Gets the texture coordinate information about the passthrough camera texture.
+	 *
+	 * @param InUV		The original UVs of on the quad. Should be an array with 8 floats.
+	 * @param OutUV		The orientated UVs that can be used to sample the passthrough camera texture and make sure it is displayed correctly.
+	 */
+	UFUNCTION(BlueprintPure, Category = "GoogleARCore|PassthroughCamera", meta = (Keywords = "googlear arcore passthrough camera uv"))
+	static void GetPassthroughCameraImageUV(const TArray<float>& InUV, TArray<float>& OutUV);
+
+	//-------------------Trackables-------------------------
+	/**
+	 * Gets a list of all valid UARPlaneGeometry objects that ARCore is currently tracking.
+	 * Planes that have entered the EARTrackingState::StoppedTracking state or for which
+	 * UARPlaneGeometry::GetSubsumedBy returns non-null will not be included.
+	 *
+	 * @param OutAnchorList		An array that contains all the valid planes detected by ARCore.
+	 */
+	UFUNCTION(BlueprintCallable, Category = "GoogleARCore|TrackablePlane", meta = (Keywords = "googlear arcore all plane"))
+	static void GetAllPlanes(TArray<UARPlaneGeometry*>& OutPlaneList);
+
+	/**
+	 * Gets a list of all valid UARTrackedPoint objects that ARCore is currently tracking.
+	 * TrackablePoint that have entered the EARTrackingState::StoppedTracking state will not be included.
+	 *
+	 * @param OutAnchorList		An array that contains all the valid planes detected by ARCore.
+	 */
+	UFUNCTION(BlueprintCallable, Category = "GoogleARCore|TrackablePoint", meta = (Keywords = "googlear arcore pose transform"))
+	static void GetAllTrackablePoints(TArray<UARTrackedPoint*>& OutTrackablePointList);
+
+	/** Template function to get all trackables from a given type. */
+	template< class T > static void GetAllTrackable(TArray<T*>& OutTrackableList);
+};
+
+/** A function library that provides static/Blueprint functions associated with most recent GoogleARCore tracking frame.*/
+UCLASS()
+class GOOGLEARCOREBASE_API UGoogleARCoreFrameFunctionLibrary : public UBlueprintFunctionLibrary
+{
+	GENERATED_BODY()
+
+public:
+	/**
+	 * Returns the current ARCore session status.
+	 *
+	 * @return	A EARSessionStatus enum that describes the session status.
+	 */
+	UFUNCTION(BlueprintPure, Category = "GoogleARCore|MotionTracking", meta = (Keywords = "googlear arcore session"))
+	static EGoogleARCoreTrackingState GetTrackingState();
+
+	/**
+	 * Gets the latest tracking pose of the ARCore device in Unreal AR Tracking Space
+	 *
+	 * Note that ARCore motion tracking has already integrated with HMD and the motion controller interface.
+	 * Use this function only if you need to implement your own tracking component.
+	 *
+	 * @param OutPose		The latest device pose.
+	 * @return				True if the pose is updated successfully for this frame.
+	 */
+	UFUNCTION(BlueprintPure, Category = "GoogleARCore|MotionTracking", meta = (Keywords = "googlear arcore pose transform"))
+	static void GetPose(FTransform& OutPose);
+
+	/**
+	 * Traces a ray from the user's device in the direction of the given location in the camera
+	 * view. Intersections with detected scene geometry are returned, sorted by distance from the
+	 * device; the nearest intersection is returned first.
+	 *
+	 * @param WorldContextObject	The world context.
+	 * @param ScreenPosition		The position on the screen to cast the ray from.
+	 * @param ARObjectType			A set of EGoogleARCoreLineTraceChannel indicate which type of line trace it should perform.
+	 * @param OutHitResults			The list of hit results sorted by distance.
+	 * @return						True if there is a hit detected.
+	 */
+	UFUNCTION(BlueprintCallable, Category = "GoogleARCore|LineTrace", meta = (WorldContext = "WorldContextObject", Keywords = "googlear arcore raycast hit"))
+	static bool ARCoreLineTrace(UObject* WorldContextObject, const FVector2D& ScreenPosition, TSet<EGoogleARCoreLineTraceChannel> TraceChannels, TArray<FARTraceResult>& OutHitResults);
+
+	/**
+	* Traces a ray along the given line. Intersections with detected scene geometry are returned, 
+	* sorted by distance from the start of the line; the nearest intersection is returned first.
+	*
+	* @param WorldContextObject	The world context.
+	* @param Start		The start of line segment.
+	* @param End		The end of line segment.
+	* @param ARObjectType			A set of EGoogleARCoreLineTraceChannel indicate which type of line trace it should perform.
+	* @param OutHitResults			The list of hit results sorted by distance.
+	* @return						True if there is a hit detected.
+	*/
+	UFUNCTION(BlueprintCallable, Category = "GoogleARCore|LineTrace", meta = (WorldContext = "WorldContextObject", Keywords = "googlear arcore raycast hit"))
+	static bool ARCoreLineTraceRay(UObject* WorldContextObject, const FVector& Start, const FVector& End, TSet<EGoogleARCoreLineTraceChannel> TraceChannels, TArray<FARTraceResult>& OutHitResults);
+
+	/**
+	 * Gets a list of UARPin objects that were changed in this frame.
+	 *
+	 * @param OutAnchorList		An array that contains the updated UARPin.
+	 */
+	UFUNCTION(BlueprintCallable, Category = "GoogleARCore|ARAnchor", meta = (Keywords = "googlear arcore Anchor"))
+	static void GetUpdatedARPins(TArray<UARPin*>& OutAnchorList);
+
+	/**
+	 * Gets a list of UARPlaneGeometry objects that were changed in this frame.
+	 *
+	 * @param OutARCorePlaneList	An array that contains the updated UARPlaneGeometry.
+	 */
+	UFUNCTION(BlueprintCallable, Category = "GoogleARCore|TrackablePlane", meta = (Keywords = "googlear arcore pose transform"))
+	static void GetUpdatedPlanes(TArray<UARPlaneGeometry*>& OutPlaneList);
+
+	/**
+	 * Gets a list of UARTrackedPoint objects that were changed in this frame.
+	 *
+	 * @param OutARCorePlaneList	An array that contains the updated UARTrackedPoint.
+	 */
+	UFUNCTION(BlueprintCallable, Category = "GoogleARCore|TrackablePoint", meta = (Keywords = "googlear arcore pose transform"))
+	static void GetUpdatedTrackablePoints(TArray<UARTrackedPoint*>& OutTrackablePointList);
+
+	/** Template function to get the updated trackables in this frame a given trackable type. */
+	template< class T > static void GetUpdatedTrackable(TArray<T*>& OutTrackableList);
+
+	/**
+	 * Gets the latest light estimation.
+	 *
+	 * @param OutLightEstimate		The struct that describes the latest light estimation.
+	 */
+	UFUNCTION(BlueprintCallable, Category = "GoogleARCore|LightEstimation", meta = (Keywords = "googlear arcore light ambient"))
+	static void GetLightEstimation(FGoogleARCoreLightEstimate& OutLightEstimate);
+
+	/**
+	 * Gets the latest point cloud that will be only available for this frame.
+	 * If you want to keep the point cloud data, you can either copy it to your own struct
+	 * or call AcquireLatestPointCloud() to avoid the copy.
+	 *
+	 * @param OutLatestPointCloud		A pointer point to the latest point cloud.
+	 * @return  An EGoogleARCoreFunctionStatus. Possible value: Success, SessionPaused, ResourceExhausted.
+	 */
+	UFUNCTION(BlueprintCallable, Category = "GoogleARCore|PointCloud", meta = (Keywords = "googlear arcore pointcloud"))
+	static EGoogleARCoreFunctionStatus GetPointCloud(UGoogleARCorePointCloud*& OutLatestPointCloud);
+
+	/**
+	 * Acquires latest point cloud. This will make the point cloud remain valid unless you call UGoogleARCrePointCloud::ReleasePointCloud().
+	 * Be aware that this function could fail if the maximal number of point cloud has been acquired.
+	 *
+	 * @param OutLatestPointCloud		A pointer point to the latest point cloud.
+	 * @return  An EGoogleARCoreFunctionStatus. Possible value: Success, SessionPaused, ResourceExhausted.
+	 */
+	UFUNCTION(BlueprintCallable, Category = "GoogleARCore|PointCloud", meta = (Keywords = "googlear arcore pointcloud"))
+	static EGoogleARCoreFunctionStatus AcquirePointCloud(UGoogleARCorePointCloud*& OutLatestPointCloud);
+
+#if PLATFORM_ANDROID
+	/**
+	 * Gets the camera metadata for the latest camera image.
+	 * Note that ACameraMetadata is a Ndk type. Include the Ndk header <camera/NdkCameraMetadata.h> to use query value from ACameraMetadata.
+	 *
+	 * @param OutCameraMetadata		A pointer to a ACameraMetadata struct which is only valid in one frame.
+	 * @return An EGoogleARCoreFunctionStatus. Possible value: Success, SessionPaused, NotAvailable.
+	 */
+	static EGoogleARCoreFunctionStatus GetCameraMetadata(const ACameraMetadata*& OutCameraMetadata);
+#endif
+
+	/**
+	 * Acquire a CPU-accessible camera image.
+	 *
+	 * @param OutLatestCameraImage    A place to store the pointer to a new UGoogleARCoreCameraImage instance.
+	 * @return An EGoogleARCoreFunctionStatus. Possible value: Success, ResourceExhausted, NotAvailable.
+	 */
+	UFUNCTION(BlueprintCallable, Category = "GoogleARCore|PassthroughCamera", meta = (Keywords = "googlear arcore passthroughcamera"))
+	static EGoogleARCoreFunctionStatus AcquireCameraImage(UGoogleARCoreCameraImage *&OutLatestCameraImage);
+
+	/**
+	 * Get the camera intrinsics for the camera image (CPU image).
+	 *
+	 * @param OutCameraIntrinsics  The output intrinsics object.
+	 */
+	UFUNCTION(BlueprintCallable, Category = "GoogleARCore|CameraIntrinsics", meta = (Keywords = "googlear arcore camera"))
+	static EGoogleARCoreFunctionStatus GetCameraImageIntrinsics(UGoogleARCoreCameraIntrinsics *&OutCameraIntrinsics);
+
+	/**
+	 * Get the camera intrinsics for the camera texture (GPU image).
+	 *
+	 * @param OutCameraIntrinsics  The output intrinsics object.
+	 */
+	UFUNCTION(BlueprintCallable, Category = "GoogleARCore|CameraIntrinsics", meta = (Keywords = "googlear arcore camera"))
+	static EGoogleARCoreFunctionStatus GetCameraTextureIntrinsics(UGoogleARCoreCameraIntrinsics *&OutCameraIntrinsics);
+
+};
diff --git a/Engine/Plugins/Experimental/AR/Google/GoogleARCore/Source/GoogleARCoreBase/Public/GoogleARCorePlaneRendererComponent.h b/Engine/Plugins/Experimental/AR/Google/GoogleARCore/Source/GoogleARCoreBase/Public/GoogleARCorePlaneRendererComponent.h
new file mode 100644
--- /dev/null
+++ b/Engine/Plugins/Experimental/AR/Google/GoogleARCore/Source/GoogleARCoreBase/Public/GoogleARCorePlaneRendererComponent.h
@@ -0,0 +1,46 @@
+// Copyright 2017 Google Inc.
+
+#pragma once
+#include "CoreMinimal.h"
+#include "Components/SceneComponent.h"
+
+#include "GoogleARCorePlaneRendererComponent.generated.h"
+
+/** A helper component that renders all the ARCore planes in the current tracking session. */
+UCLASS(Experimental, ClassGroup = (GoogleARCore), meta = (BlueprintSpawnableComponent))
+class GOOGLEARCOREBASE_API UGoogleARCorePlaneRendererComponent : public USceneComponent
+{
+	GENERATED_BODY()
+public:
+
+	/** Render the plane quad when set to true. */
+	UPROPERTY(EditAnywhere, BlueprintReadWrite, Category = "GoogleARCore|TrackablePlaneRenderer")
+	bool bRenderPlane;
+
+	/** Render the plane boundary polygon lines when set to true. */
+	UPROPERTY(EditAnywhere, BlueprintReadWrite, Category = "GoogleARCore|TrackablePlaneRenderer")
+	bool bRenderBoundaryPolygon;
+
+	/** The color of the plane. */
+	UPROPERTY(EditAnywhere, BlueprintReadWrite, Category = "GoogleARCore|TrackablePlaneRenderer")
+	FColor PlaneColor;
+
+	/** The color of the boundary polygon. */
+	UPROPERTY(EditAnywhere, BlueprintReadWrite, Category = "GoogleARCore|TrackablePlaneRenderer")
+	FColor BoundaryPolygonColor;
+
+	/** The line thickness for the plan boundary polygon. */
+	UPROPERTY(EditAnywhere, BlueprintReadWrite, Category = "GoogleARCore|TrackablePlaneRenderer")
+	float BoundaryPolygonThickness;
+
+	UGoogleARCorePlaneRendererComponent();
+
+	/** Function called every frame on this Component. */
+	void TickComponent(float DeltaTime, enum ELevelTick TickType, FActorComponentTickFunction * ThisTickFunction) override;
+
+private:
+	TArray<int> PlaneIndices;
+	TArray<FVector> PlaneVertices;
+
+	void DrawPlanes();
+};
diff --git a/Engine/Plugins/Experimental/AR/Google/GoogleARCore/Source/GoogleARCoreBase/Public/GoogleARCorePointCloudRendererComponent.h b/Engine/Plugins/Experimental/AR/Google/GoogleARCore/Source/GoogleARCoreBase/Public/GoogleARCorePointCloudRendererComponent.h
new file mode 100644
--- /dev/null
+++ b/Engine/Plugins/Experimental/AR/Google/GoogleARCore/Source/GoogleARCoreBase/Public/GoogleARCorePointCloudRendererComponent.h
@@ -0,0 +1,34 @@
+// Copyright 2017 Google Inc.
+
+#pragma once
+
+#include "CoreMinimal.h"
+#include "Components/SceneComponent.h"
+
+#include "GoogleARCorePointCloudRendererComponent.generated.h"
+
+/** A helper component that renders the latest point cloud from the ARCore tracking session. */
+UCLASS(Experimental, ClassGroup = (GoogleARCore), meta = (BlueprintSpawnableComponent))
+class GOOGLEARCOREBASE_API UGoogleARCorePointCloudRendererComponent : public USceneComponent
+{
+	GENERATED_BODY()
+public:
+	/** The color of the point. */
+	UPROPERTY(EditAnywhere, BlueprintReadWrite, Category = "GoogleARCore|PointCloudRenderer")
+	FColor PointColor;
+
+	/** The size of the point. */
+	UPROPERTY(EditAnywhere, BlueprintReadWrite, Category = "GoogleARCore|PointCloudRenderer")
+	float PointSize;
+
+	UGoogleARCorePointCloudRendererComponent();
+
+	/** Function called on every frame on this Component. */
+	void TickComponent(float DeltaTime, enum ELevelTick TickType, FActorComponentTickFunction * ThisTickFunction) override;
+
+private:
+	TArray<FVector> PointCloudInWorldSpace;
+	double PreviousPointCloudTimestamp;
+
+	void DrawPointCloud();
+};
diff --git a/Engine/Plugins/Experimental/AR/Google/GoogleARCore/Source/GoogleARCoreRendering/GoogleARCoreRendering.Build.cs b/Engine/Plugins/Experimental/AR/Google/GoogleARCore/Source/GoogleARCoreRendering/GoogleARCoreRendering.Build.cs
new file mode 100644
--- /dev/null
+++ b/Engine/Plugins/Experimental/AR/Google/GoogleARCore/Source/GoogleARCoreRendering/GoogleARCoreRendering.Build.cs
@@ -0,0 +1,34 @@
+// Copyright 2017 Google Inc.
+
+using UnrealBuildTool;
+
+public class GoogleARCoreRendering : ModuleRules
+{
+	public GoogleARCoreRendering(ReadOnlyTargetRules Target) : base(Target)
+	{
+		PrivateIncludePaths.AddRange(
+			new string[] {
+				"GoogleARCoreRendering/Private",
+				"GoogleARCoreBase/Private",
+				"../../../../Source/Runtime/Renderer/Private",
+				// ... add other private include paths required here ...
+			}
+			);
+		PrivateDependencyModuleNames.AddRange(
+			new string[]
+			{
+				"Core",
+				"Json",
+				"CoreUObject",
+				"Engine",
+				"RHI",
+				"Engine",
+				"Renderer",
+				"RenderCore",
+				"ShaderCore",
+				// ... add private dependencies that you statically link with here ...
+			}
+			);
+		bFasterWithoutUnity = true;
+	}
+}
\ No newline at end of file
diff --git a/Engine/Plugins/Experimental/AR/Google/GoogleARCore/Source/GoogleARCoreRendering/Private/GoogleARCorePassthroughCameraRenderer.cpp b/Engine/Plugins/Experimental/AR/Google/GoogleARCore/Source/GoogleARCoreRendering/Private/GoogleARCorePassthroughCameraRenderer.cpp
new file mode 100644
--- /dev/null
+++ b/Engine/Plugins/Experimental/AR/Google/GoogleARCore/Source/GoogleARCoreRendering/Private/GoogleARCorePassthroughCameraRenderer.cpp
@@ -0,0 +1,327 @@
+// Copyright 2017 Google Inc.
+
+#include "GoogleARCorePassthroughCameraRenderer.h"
+#include "ScreenRendering.h"
+#include "RendererInterface.h"
+#include "EngineModule.h"
+#include "PipelineStateCache.h"
+#include "ClearQuad.h"
+#include "Containers/DynamicRHIResourceArray.h"
+#include "EngineGlobals.h"
+#include "Engine/Engine.h"
+#include "SceneUtils.h"
+#include "Engine/TextureRenderTarget2D.h"
+#include "PostProcess/SceneFilterRendering.h"
+#include "PostProcessParameters.h"
+#include "MaterialShader.h"
+#include "MaterialShaderType.h"
+#include "ExternalTexture.h"
+#include "GoogleARCorePassthroughCameraExternalTextureGuid.h"
+#include "GoogleARCoreAndroidHelper.h"
+
+FGoogleARCorePassthroughCameraRenderer::FGoogleARCorePassthroughCameraRenderer()
+	: OverlayQuadUVs{ 0.0f, 0.0f, 0.0f, 1.0f, 1.0f, 0.0f, 1.0f, 1.0f }
+	, bInitialized(false)
+	, VideoTexture(nullptr)
+	, bMaterialInitialized(false)
+	, DefaultOverlayMaterial(nullptr)
+	, OverrideOverlayMaterial(nullptr)
+	, RenderingOverlayMaterial(nullptr)
+{
+}
+
+void FGoogleARCorePassthroughCameraRenderer::SetDefaultCameraOverlayMaterial(UMaterialInterface* InDefaultCameraOverlayMaterial)
+{
+	DefaultOverlayMaterial = InDefaultCameraOverlayMaterial;
+}
+
+void FGoogleARCorePassthroughCameraRenderer::InitializeOverlayMaterial()
+{
+	if (RenderingOverlayMaterial != nullptr)
+		return;
+
+	SetDefaultCameraOverlayMaterial(GetDefault<UGoogleARCoreCameraOverlayMaterialLoader>()->DefaultCameraOverlayMaterial);
+	ResetOverlayMaterialToDefault();
+}
+
+void FGoogleARCorePassthroughCameraRenderer::SetOverlayMaterialInstance(UMaterialInterface* NewMaterialInstance)
+{
+	if (NewMaterialInstance != nullptr)
+	{
+		OverrideOverlayMaterial = NewMaterialInstance;
+
+		ENQUEUE_UNIQUE_RENDER_COMMAND_ONEPARAMETER(
+        UseOverrideOverlayMaterial,
+        FGoogleARCorePassthroughCameraRenderer*, VideoOverlayRendererRHIPtr, this,
+        {
+            VideoOverlayRendererRHIPtr->RenderingOverlayMaterial = VideoOverlayRendererRHIPtr->OverrideOverlayMaterial;
+        });
+	}
+}
+
+void FGoogleARCorePassthroughCameraRenderer::ResetOverlayMaterialToDefault()
+{
+	ENQUEUE_UNIQUE_RENDER_COMMAND_ONEPARAMETER(
+        UseDefaultOverlayMaterial,
+        FGoogleARCorePassthroughCameraRenderer*, VideoOverlayRendererRHIPtr, this,
+        {
+            VideoOverlayRendererRHIPtr->RenderingOverlayMaterial = VideoOverlayRendererRHIPtr->DefaultOverlayMaterial;
+        }
+    );
+}
+
+void FGoogleARCorePassthroughCameraRenderer::InitializeRenderer_RenderThread(FTextureRHIRef ExternalTexture)
+{
+	if (bInitialized)
+		return;
+
+	// Initialize Index buffer;
+	const uint16 Indices[] = { 0, 1, 2, 2, 1, 3};
+
+	TResourceArray<uint16, INDEXBUFFER_ALIGNMENT> IndexBuffer;
+	uint32 NumIndices = ARRAY_COUNT(Indices);
+	IndexBuffer.AddUninitialized(NumIndices);
+	FMemory::Memcpy(IndexBuffer.GetData(), Indices, NumIndices * sizeof(uint16));
+
+	// Create index buffer. Fill buffer with initial data upon creation
+	FRHIResourceCreateInfo CreateInfo(&IndexBuffer);
+	OverlayIndexBufferRHI = RHICreateIndexBuffer(sizeof(uint16), IndexBuffer.GetResourceDataSize(), BUF_Static, CreateInfo);
+
+	VideoTexture = ExternalTexture;
+
+	FSamplerStateInitializerRHI SamplerStateInitializer(SF_Point, AM_Clamp, AM_Clamp, AM_Clamp);
+	FSamplerStateRHIRef SamplerStateRHI = RHICreateSamplerState(SamplerStateInitializer);
+
+	FExternalTextureRegistry::Get().RegisterExternalTexture(GoogleARCorePassthroughCameraExternalTextureGuid, VideoTexture, SamplerStateRHI);
+
+	bInitialized = true;
+}
+
+void FGoogleARCorePassthroughCameraRenderer::UpdateOverlayUVCoordinate_RenderThread(TArray<float>& InOverlayUVs, ARCoreDisplayRotation DisplayRotation)
+{
+	check(InOverlayUVs.Num() == 8);
+
+	// It seems very likely that this is papering over up some underlying problem with the camera image orientation.
+	bool bIsLandscape = (DisplayRotation == ARCoreDisplayRotation::Rotation0 || DisplayRotation == ARCoreDisplayRotation::Rotation180);
+	bool bNeedToFlipCameraImageHorizontally = !bIsLandscape && RHINeedsToSwitchVerticalAxis(GShaderPlatformForFeatureLevel[GMaxRHIFeatureLevel]) && !IsMobileHDR();
+	bool bFlipCameraImageVertically = bIsLandscape && RHINeedsToSwitchVerticalAxis(GShaderPlatformForFeatureLevel[GMaxRHIFeatureLevel]) && !IsMobileHDR();
+	bool bDiagonalFlip = bIsLandscape && IsMobileHDR();
+
+	if (bFlipCameraImageVertically)
+	{
+		InOverlayUVs.SwapMemory(0, 4);
+		InOverlayUVs.SwapMemory(1, 5);
+		InOverlayUVs.SwapMemory(2, 6);
+		InOverlayUVs.SwapMemory(3, 7);
+	}
+	else if (bNeedToFlipCameraImageHorizontally)
+	{
+		InOverlayUVs.SwapMemory(0, 2);
+		InOverlayUVs.SwapMemory(1, 3);
+		InOverlayUVs.SwapMemory(4, 6);
+		InOverlayUVs.SwapMemory(5, 7);
+	}
+	else if (bDiagonalFlip)
+	{
+		InOverlayUVs.SwapMemory(0, 6);
+		InOverlayUVs.SwapMemory(1, 7);
+		InOverlayUVs.SwapMemory(4, 2);
+		InOverlayUVs.SwapMemory(5, 3);
+	}
+	{
+		if (OverlayVertexBufferRHI.IsValid())
+		{
+			OverlayVertexBufferRHI.SafeRelease();
+		}
+
+		TResourceArray<FFilterVertex, VERTEXBUFFER_ALIGNMENT> Vertices;
+		Vertices.SetNumUninitialized(4);
+
+		// Unreal uses reversed z. 0 is the farthest.
+		Vertices[0].Position = FVector4(0, 1, 0, 1);
+		Vertices[0].UV = FVector2D(InOverlayUVs[0], InOverlayUVs[1]);
+
+		Vertices[1].Position = FVector4(0, 0, 0, 1);
+		Vertices[1].UV = FVector2D(InOverlayUVs[2], InOverlayUVs[3]);
+
+		Vertices[2].Position = FVector4(1, 1, 0, 1);
+		Vertices[2].UV = FVector2D(InOverlayUVs[4], InOverlayUVs[5]);
+
+		Vertices[3].Position = FVector4(1, 0, 0, 1);
+		Vertices[3].UV = FVector2D(InOverlayUVs[6], InOverlayUVs[7]);
+
+		// Create vertex buffer. Fill buffer with initial data upon creation
+		FRHIResourceCreateInfo CreateInfo(&Vertices);
+		OverlayVertexBufferRHI = RHICreateVertexBuffer(Vertices.GetResourceDataSize(), BUF_Static, CreateInfo);
+	}
+}
+
+// We use something similar to the PostProcessMaterial to render the color camera overlay.
+class FGoogleARCoreCameraOverlayVS : public FMaterialShader
+{
+	DECLARE_SHADER_TYPE(FGoogleARCoreCameraOverlayVS, Material);
+public:
+
+	static bool ShouldCompilePermutation(EShaderPlatform Platform, const FMaterial* Material)
+	{
+		return Material->GetMaterialDomain() == MD_PostProcess && IsMobilePlatform(Platform);
+	}
+
+	static void ModifyCompilationEnvironment(EShaderPlatform Platform, const class FMaterial* Material, FShaderCompilerEnvironment& OutEnvironment)
+	{
+		FMaterialShader::ModifyCompilationEnvironment(Platform, OutEnvironment);
+
+		OutEnvironment.SetDefine(TEXT("POST_PROCESS_MATERIAL"), 1);
+
+		OutEnvironment.SetDefine(TEXT("POST_PROCESS_MATERIAL_BEFORE_TONEMAP"), (Material->GetBlendableLocation() != BL_AfterTonemapping) ? 1 : 0);
+	}
+
+	FGoogleARCoreCameraOverlayVS( )	{ }
+	FGoogleARCoreCameraOverlayVS(const ShaderMetaType::CompiledShaderInitializerType& Initializer)
+		: FMaterialShader(Initializer)
+	{
+	}
+
+	void SetParameters(FRHICommandList& RHICmdList, const FSceneView View)
+	{
+		const FVertexShaderRHIParamRef ShaderRHI = GetVertexShader();
+		FMaterialShader::SetViewParameters(RHICmdList, ShaderRHI, View, View.ViewUniformBuffer);
+	}
+
+	virtual bool Serialize(FArchive& Ar) override
+	{
+		bool bShaderHasOutdatedParameters = FMaterialShader::Serialize(Ar);
+		return bShaderHasOutdatedParameters;
+	}
+};
+
+IMPLEMENT_MATERIAL_SHADER_TYPE(,FGoogleARCoreCameraOverlayVS,TEXT("/Engine/Private/PostProcessMaterialShaders.usf"),TEXT("MainVS_ES2"),SF_Vertex);
+
+class FGoogleARCoreCameraOverlayPS : public FMaterialShader
+{
+	DECLARE_SHADER_TYPE(FGoogleARCoreCameraOverlayPS, Material);
+public:
+
+	static bool ShouldCompilePermutation(EShaderPlatform Platform, const FMaterial* Material)
+	{
+		return Material->GetMaterialDomain() == MD_PostProcess && IsMobilePlatform(Platform);
+	}
+
+	static void ModifyCompilationEnvironment(EShaderPlatform Platform, const class FMaterial* Material, FShaderCompilerEnvironment& OutEnvironment)
+	{
+		FMaterialShader::ModifyCompilationEnvironment(Platform, OutEnvironment);
+
+		OutEnvironment.SetDefine(TEXT("POST_PROCESS_MATERIAL"), 1);
+		OutEnvironment.SetDefine(TEXT("OUTPUT_GAMMA_SPACE"), IsMobileHDR() ? 0 : 1);
+		OutEnvironment.SetDefine(TEXT("POST_PROCESS_MATERIAL_BEFORE_TONEMAP"), (Material->GetBlendableLocation() != BL_AfterTonemapping) ? 1 : 0);
+	}
+
+	FGoogleARCoreCameraOverlayPS() {}
+	FGoogleARCoreCameraOverlayPS(const ShaderMetaType::CompiledShaderInitializerType& Initializer) :
+		FMaterialShader(Initializer)
+	{
+		for (uint32 InputIter = 0; InputIter < ePId_Input_MAX; ++InputIter)
+		{
+			PostprocessInputParameter[InputIter].Bind(Initializer.ParameterMap, *FString::Printf(TEXT("PostprocessInput%d"), InputIter));
+			PostprocessInputParameterSampler[InputIter].Bind(Initializer.ParameterMap, *FString::Printf(TEXT("PostprocessInput%dSampler"), InputIter));
+		}
+	}
+
+	void SetParameters(FRHICommandList& RHICmdList, const FSceneView View, const FMaterialRenderProxy* Material)
+	{
+		const FPixelShaderRHIParamRef ShaderRHI = GetPixelShader();
+
+		FMaterialShader::SetParameters(RHICmdList, ShaderRHI, Material, *Material->GetMaterial(View.GetFeatureLevel()), View, View.ViewUniformBuffer, ESceneTextureSetupMode::None);
+
+		for (uint32 InputIter = 0; InputIter < ePId_Input_MAX; ++InputIter)
+		{
+			if (PostprocessInputParameter[InputIter].IsBound())
+			{
+				SetTextureParameter(
+					RHICmdList,
+					ShaderRHI,
+					PostprocessInputParameter[InputIter],
+					PostprocessInputParameterSampler[InputIter],
+					TStaticSamplerState<>::GetRHI(),
+					GBlackTexture->TextureRHI);
+			}
+		}
+	}
+
+	virtual bool Serialize(FArchive& Ar) override
+	{
+		bool bShaderHasOutdatedParameters = FMaterialShader::Serialize(Ar);
+		return bShaderHasOutdatedParameters;
+	}
+
+private:
+	FShaderResourceParameter PostprocessInputParameter[ePId_Input_MAX];
+	FShaderResourceParameter PostprocessInputParameterSampler[ePId_Input_MAX];
+};
+
+IMPLEMENT_MATERIAL_SHADER_TYPE(,FGoogleARCoreCameraOverlayPS,TEXT("/Engine/Private/PostProcessMaterialShaders.usf"),TEXT("MainPS_ES2"),SF_Pixel);
+
+void FGoogleARCorePassthroughCameraRenderer::RenderVideoOverlay_RenderThread(FRHICommandListImmediate& RHICmdList, FSceneView& InView)
+{
+#if PLATFORM_ANDROID
+	if (RenderingOverlayMaterial == nullptr || !RenderingOverlayMaterial->IsValidLowLevel())
+	{
+		return;
+	}
+
+	const auto FeatureLevel = InView.GetFeatureLevel();
+	IRendererModule& RendererModule = GetRendererModule();
+
+	if (FeatureLevel <= ERHIFeatureLevel::ES3_1)
+	{
+		const FMaterial* CameraMaterial = RenderingOverlayMaterial->GetRenderProxy(false)->GetMaterial(FeatureLevel);
+		const FMaterialShaderMap* MaterialShaderMap = CameraMaterial->GetRenderingThreadShaderMap();
+
+		FGoogleARCoreCameraOverlayPS* PixelShader = MaterialShaderMap->GetShader<FGoogleARCoreCameraOverlayPS>();
+		FGoogleARCoreCameraOverlayVS* VertexShader = MaterialShaderMap->GetShader<FGoogleARCoreCameraOverlayVS>();
+
+		FGraphicsPipelineStateInitializer GraphicsPSOInit;
+
+		GraphicsPSOInit.BlendState = TStaticBlendState<>::GetRHI();
+		GraphicsPSOInit.RasterizerState = TStaticRasterizerState<>::GetRHI();
+		GraphicsPSOInit.DepthStencilState = TStaticDepthStencilState<false, CF_DepthNearOrEqual>::GetRHI();
+
+		GraphicsPSOInit.BoundShaderState.VertexDeclarationRHI = RendererModule.GetFilterVertexDeclaration().VertexDeclarationRHI;
+		GraphicsPSOInit.BoundShaderState.VertexShaderRHI = GETSAFERHISHADER_VERTEX(VertexShader);
+		GraphicsPSOInit.BoundShaderState.PixelShaderRHI = GETSAFERHISHADER_PIXEL(PixelShader);
+		GraphicsPSOInit.PrimitiveType = PT_TriangleList;
+
+		SetGraphicsPipelineState(RHICmdList, GraphicsPSOInit, EApplyRendertargetOption::DoNothing);
+
+		VertexShader->SetParameters(RHICmdList, InView);
+		PixelShader->SetParameters(RHICmdList, InView, RenderingOverlayMaterial->GetRenderProxy(false));
+
+		FIntPoint ViewSize = InView.UnscaledViewRect.Size();
+
+		FDrawRectangleParameters Parameters;
+		Parameters.PosScaleBias = FVector4(ViewSize.X, ViewSize.Y, 0, 0);
+		Parameters.UVScaleBias = FVector4(1.0f, 1.0f, 0.0f, 0.0f);
+
+		Parameters.InvTargetSizeAndTextureSize = FVector4(
+			1.0f / ViewSize.X, 1.0f / ViewSize.Y,
+			1.0f, 1.0f);
+
+		SetUniformBufferParameterImmediate(RHICmdList, VertexShader->GetVertexShader(), VertexShader->GetUniformBufferParameter<FDrawRectangleParameters>(), Parameters);
+
+		if (OverlayVertexBufferRHI.IsValid() && OverlayIndexBufferRHI.IsValid())
+		{
+			RHICmdList.SetStreamSource(0, OverlayVertexBufferRHI, 0);
+			RHICmdList.DrawIndexedPrimitive(
+				OverlayIndexBufferRHI,
+				PT_TriangleList,
+				/*BaseVertexIndex=*/ 0,
+				/*MinIndex=*/ 0,
+				/*NumVertices=*/ 4,
+				/*StartIndex=*/ 0,
+				/*NumPrimitives=*/ 2,
+				/*NumInstances=*/ 1
+			);
+		}
+	}
+#endif
+}
