Grouped Apple/Google AR plugins under an AR directory in Experimental to improve the organization of various AR plugins

diff --git a/Engine/Plugins/Experimental/AR/Google/GoogleARCore/Source/GoogleARCoreRendering/Private/MaterialExpressionGoogleARCorePassthroughCamera.cpp b/Engine/Plugins/Experimental/AR/Google/GoogleARCore/Source/GoogleARCoreRendering/Private/MaterialExpressionGoogleARCorePassthroughCamera.cpp
new file mode 100644
--- /dev/null
+++ b/Engine/Plugins/Experimental/AR/Google/GoogleARCore/Source/GoogleARCoreRendering/Private/MaterialExpressionGoogleARCorePassthroughCamera.cpp
@@ -0,0 +1,32 @@
+// Copyright 2017 Google Inc.
+
+#include "MaterialExpressionGoogleARCorePassthroughCamera.h"
+#include "Engine/Engine.h"
+#include "EngineGlobals.h"
+#include "MaterialCompiler.h"
+#include "GoogleARCorePassthroughCameraExternalTextureGuid.h"
+
+UMaterialExpressionGoogleARCorePassthroughCamera::UMaterialExpressionGoogleARCorePassthroughCamera(const FObjectInitializer& ObjectInitializer)
+	: Super(ObjectInitializer)
+{
+}
+
+#if WITH_EDITOR
+int32 UMaterialExpressionGoogleARCorePassthroughCamera::Compile(class FMaterialCompiler* Compiler, int32 OutputIndex)
+{
+	return Compiler->TextureSample(
+		Compiler->ExternalTexture(GoogleARCorePassthroughCameraExternalTextureGuid),
+		Coordinates.GetTracedInput().Expression ? Coordinates.Compile(Compiler) : Compiler->TextureCoordinate(ConstCoordinate, false, false),
+		EMaterialSamplerType::SAMPLERTYPE_Color);
+}
+
+int32 UMaterialExpressionGoogleARCorePassthroughCamera::CompilePreview(class FMaterialCompiler* Compiler, int32 OutputIndex)
+{
+	return INDEX_NONE;
+}
+
+void UMaterialExpressionGoogleARCorePassthroughCamera::GetCaption(TArray<FString>& OutCaptions) const
+{
+	OutCaptions.Add(TEXT("GoogleARCore Passthrough Camera"));
+}
+#endif
diff --git a/Engine/Plugins/Experimental/AR/Google/GoogleARCore/Source/GoogleARCoreRendering/Public/GoogleARCorePassthroughCameraRenderer.h b/Engine/Plugins/Experimental/AR/Google/GoogleARCore/Source/GoogleARCoreRendering/Public/GoogleARCorePassthroughCameraRenderer.h
new file mode 100644
--- /dev/null
+++ b/Engine/Plugins/Experimental/AR/Google/GoogleARCore/Source/GoogleARCoreRendering/Public/GoogleARCorePassthroughCameraRenderer.h
@@ -0,0 +1,76 @@
+// Copyright 2017 Google Inc.
+
+#pragma once
+
+#include "CoreMinimal.h"
+#include "UObject/ConstructorHelpers.h"
+#include "RHI.h"
+#include "RenderResource.h"
+#include "ShaderParameters.h"
+#include "Shader.h"
+#include "GlobalShader.h"
+#include "ShaderParameterUtils.h"
+#include "Materials/Material.h"
+#include "Materials/MaterialInterface.h"
+#include "Materials/MaterialInstanceDynamic.h"
+#include "Engine/Texture.h"
+
+#include "GoogleARCorePassthroughCameraRenderer.generated.h"
+
+enum class ARCoreDisplayRotation : int32;
+
+/** A helper class that is used to load the GoogleARCorePassthroughCameraMaterial from its default object. */
+UCLASS()
+class GOOGLEARCORERENDERING_API UGoogleARCoreCameraOverlayMaterialLoader : public UObject
+{
+	GENERATED_BODY()
+
+public:
+	/** A pointer to the camera overlay material that will be used to render the passthrough camera texture as background. */
+	UPROPERTY()
+		UMaterialInterface* DefaultCameraOverlayMaterial;
+
+	UGoogleARCoreCameraOverlayMaterialLoader()
+	{
+		static ConstructorHelpers::FObjectFinder<UMaterialInterface> DefaultOverlayMaterialRef(TEXT("/GoogleARCore/GoogleARCorePassthroughCameraMaterial.GoogleARCorePassthroughCameraMaterial"));
+		DefaultCameraOverlayMaterial = DefaultOverlayMaterialRef.Object;
+
+		UMaterialInstanceDynamic *DynMat = UMaterialInstanceDynamic::Create(
+			DefaultOverlayMaterialRef.Object, this, FName("GoogleARCorePassthroughCameraMaterial_Dynamic"));
+
+		DefaultCameraOverlayMaterial = DynMat;
+	}
+};
+
+class FRHICommandListImmediate;
+
+class GOOGLEARCORERENDERING_API FGoogleARCorePassthroughCameraRenderer
+{
+public:
+	FGoogleARCorePassthroughCameraRenderer();
+
+	void SetDefaultCameraOverlayMaterial(UMaterialInterface* InDefaultCameraOverlayMaterial);
+
+	void InitializeOverlayMaterial();
+
+	void SetOverlayMaterialInstance(UMaterialInterface* NewMaterialInstance);
+	void ResetOverlayMaterialToDefault();
+
+	void InitializeRenderer_RenderThread(FTextureRHIRef ExternalTexture);
+
+	void UpdateOverlayUVCoordinate_RenderThread(TArray<float>& InOverlayUVs, ARCoreDisplayRotation DisplayRotation);
+	void RenderVideoOverlay_RenderThread(FRHICommandListImmediate& RHICmdList, FSceneView& InView);
+
+	const TArray<float> OverlayQuadUVs;
+private:
+	bool bInitialized;
+	FIndexBufferRHIRef OverlayIndexBufferRHI;
+	FVertexBufferRHIRef OverlayVertexBufferRHI;
+	FTextureRHIRef VideoTexture;
+	float OverlayTextureUVs[8];
+	bool bMaterialInitialized;
+	UMaterialInterface* DefaultOverlayMaterial;
+	UMaterialInterface* OverrideOverlayMaterial;
+	UMaterialInterface* RenderingOverlayMaterial;
+};
+
diff --git a/Engine/Plugins/Experimental/AR/Google/GoogleARCoreServices/Source/GoogleARCoreServices/GoogleARCoreServices.Build.cs b/Engine/Plugins/Experimental/AR/Google/GoogleARCoreServices/Source/GoogleARCoreServices/GoogleARCoreServices.Build.cs
new file mode 100644
--- /dev/null
+++ b/Engine/Plugins/Experimental/AR/Google/GoogleARCoreServices/Source/GoogleARCoreServices/GoogleARCoreServices.Build.cs
@@ -0,0 +1,55 @@
+// Copyright 2018 Google Inc.
+
+
+using System.IO;
+
+
+namespace UnrealBuildTool.Rules
+{
+	public class GoogleARCoreServices : ModuleRules
+	{
+		public GoogleARCoreServices(ReadOnlyTargetRules Target) : base(Target)
+		{
+			PrivateIncludePaths.AddRange(
+				new string[]
+				{
+					"GoogleARCoreServices/Private",
+				}
+			);
+			
+			PublicDependencyModuleNames.AddRange(
+					new string[]
+					{
+						"HeadMountedDisplay",
+						"AugmentedReality",
+					}
+			);
+
+			PrivateDependencyModuleNames.AddRange(
+				new string[]
+				{
+					"Core",
+					"CoreUObject",
+					"Engine",
+					"GoogleARCoreSDK",
+				}
+			);
+
+			PrivateIncludePathModuleNames.AddRange(
+				new string[]
+				{
+					"Settings" // For editor settings panel.
+				}
+			);
+
+
+			if (Target.Platform == UnrealTargetPlatform.Android)
+			{
+				// Register Plugin Language
+				AdditionalPropertiesForReceipt.Add("AndroidPlugin", Path.Combine(ModuleDirectory, "GoogleARCoreServices_APL.xml"));
+			}
+
+			bFasterWithoutUnity = false;
+		}
+	}
+}
diff --git a/Engine/Plugins/Experimental/AR/Google/GoogleARCoreServices/Source/GoogleARCoreServices/Private/GoogleARCoreServicesManager.h b/Engine/Plugins/Experimental/AR/Google/GoogleARCoreServices/Source/GoogleARCoreServices/Private/GoogleARCoreServicesManager.h
new file mode 100644
--- /dev/null
+++ b/Engine/Plugins/Experimental/AR/Google/GoogleARCoreServices/Source/GoogleARCoreServices/Private/GoogleARCoreServicesManager.h
@@ -0,0 +1,44 @@
+// Copyright 2018 Google Inc.
+
+#pragma once
+
+#include "CoreMinimal.h"
+#include "ARSystem.h"
+
+#include "GoogleARCoreServicesTypes.h"
+#include "GoogleARCoreCloudARPinManager.h"
+
+class FGoogleARCoreServicesManager
+{
+
+public:
+	FGoogleARCoreServicesManager();
+
+	~FGoogleARCoreServicesManager();
+
+	bool ConfigGoogleARCoreServices(FGoogleARCoreServicesConfig& ServiceConfig);
+
+	UCloudARPin* CreateAndHostCloudARPin(UARPin* ARPinToHost, EARPinCloudTaskResult& OutTaskResult);
+
+	UCloudARPin* ResolveAncCreateCloudARPin(FString CloudId, EARPinCloudTaskResult& OutTaskResult);
+
+	void RemoveCloudARPin(UCloudARPin* PinToRemove);
+
+	TArray<UCloudARPin*> GetAllCloudARPin();
+
+private:
+	bool InitARSystem();
+
+	EARPinCloudTaskResult CheckCloudTaskError();
+
+	void OnARSessionStarted();
+	void OnWorldTickStart(ELevelTick TickType, float DeltaTime);
+
+	bool bHasValidARSystem;
+	bool bCloudARPinEnabled;
+	FGoogleARCoreServicesConfig CurrentServicesConfig;
+
+	TSharedPtr<FARSupportInterface , ESPMode::ThreadSafe> ARSystem;
+	TUniquePtr<FGoogleARCoreCloudARPinManager> CloudARPinManager;
+};
+
