Don't allow camera to rotate straight vertical up/down just because ARCore plugin is enabled, instead only allow if we're actually in an ARSession with "Head" tracking enabled.

diff --git a/Engine/Plugins/Runtime/AR/Google/GoogleARCore/Source/GoogleARCoreBase/Private/GoogleARCoreXRTrackingSystem.cpp b/Engine/Plugins/Runtime/AR/Google/GoogleARCore/Source/GoogleARCoreBase/Private/GoogleARCoreXRTrackingSystem.cpp
--- a/Engine/Plugins/Runtime/AR/Google/GoogleARCore/Source/GoogleARCoreBase/Private/GoogleARCoreXRTrackingSystem.cpp
+++ b/Engine/Plugins/Runtime/AR/Google/GoogleARCore/Source/GoogleARCoreBase/Private/GoogleARCoreXRTrackingSystem.cpp
@@ -44,7 +44,8 @@ FName FGoogleARCoreXRTrackingSystem::GetSystemName() const
 bool FGoogleARCoreXRTrackingSystem::IsHeadTrackingAllowed() const
 {
 #if PLATFORM_ANDROID
-	return true;
+	return FGoogleARCoreDevice::GetInstance()->GetSessionStatus().Status == EARSessionStatus::Running &&
+		FGoogleARCoreDevice::GetInstance()->GetARSystem()->GetSessionConfig().ShouldEnableCameraTracking();
 #else
 	return false;
 #endif
