Minor revert to logic to avoid static analysis errors in AR Core

diff --git a/Engine/Plugins/Runtime/AR/Google/GoogleARCore/Source/GoogleARCoreBase/Private/GoogleARCoreAugmentedImageDatabase.cpp b/Engine/Plugins/Runtime/AR/Google/GoogleARCore/Source/GoogleARCoreBase/Private/GoogleARCoreAugmentedImageDatabase.cpp
--- a/Engine/Plugins/Runtime/AR/Google/GoogleARCore/Source/GoogleARCoreBase/Private/GoogleARCoreAugmentedImageDatabase.cpp
+++ b/Engine/Plugins/Runtime/AR/Google/GoogleARCore/Source/GoogleARCoreBase/Private/GoogleARCoreAugmentedImageDatabase.cpp
@@ -48,7 +48,7 @@ int UGoogleARCoreAugmentedImageDatabase::AddRuntimeAugmentedImageFromTexture(UTe
 				GrayscaleBuffer[i] = 0.2126 * R + 0.7152 * G + 0.0722 * B;
 			}
 		}
-		else if(PixelFormat == EPixelFormat::PF_G8)
+		else
 		{
 			ensureMsgf(RawImageData->GetBulkDataSize() == ImageWidth * ImageHeight,
 				TEXT("Unsupported texture data in UGoogleARCoreAugmentedImageDatabase::AddRuntimeAugmentedImage"));
diff --git a/Engine/Plugins/Runtime/AR/Google/GoogleARCore/Source/GoogleARCoreBase/Private/GoogleARCoreXRTrackingSystem.cpp b/Engine/Plugins/Runtime/AR/Google/GoogleARCore/Source/GoogleARCoreBase/Private/GoogleARCoreXRTrackingSystem.cpp
--- a/Engine/Plugins/Runtime/AR/Google/GoogleARCore/Source/GoogleARCoreBase/Private/GoogleARCoreXRTrackingSystem.cpp
+++ b/Engine/Plugins/Runtime/AR/Google/GoogleARCore/Source/GoogleARCoreBase/Private/GoogleARCoreXRTrackingSystem.cpp
@@ -370,7 +370,7 @@ bool FGoogleARCoreXRTrackingSystem::OnAddRuntimeCandidateImage(UARSessionConfig*
 				GrayscaleBuffer[i] = 0.2126 * R + 0.7152 * G + 0.0722 * B;
 			}
 		}
-		else if (PixelFormat == EPixelFormat::PF_G8)
+		else
 		{
 			ensureMsgf(RawImageData->GetBulkDataSize() == ImageWidth * ImageHeight,
 				TEXT("Unsupported texture data when adding runtime candidate image."));
