Hopefully fix missing USE_ANDROID_OPENGL definition
#rnx

diff --git a/Engine/Plugins/Runtime/AR/Google/GoogleARCore/Source/GoogleARCoreBase/GoogleARCoreBase.Build.cs b/Engine/Plugins/Runtime/AR/Google/GoogleARCore/Source/GoogleARCoreBase/GoogleARCoreBase.Build.cs
--- a/Engine/Plugins/Runtime/AR/Google/GoogleARCore/Source/GoogleARCoreBase/GoogleARCoreBase.Build.cs
+++ b/Engine/Plugins/Runtime/AR/Google/GoogleARCore/Source/GoogleARCoreBase/GoogleARCoreBase.Build.cs
@@ -42,6 +42,7 @@ namespace UnrealBuildTool.Rules
 					"GoogleARCoreRendering",
 					"GoogleARCoreSDK",
 					"OpenGL",
+					"OpenGLDrv",
 					"ProceduralMeshComponent",
 					"UElibPNG",
 					"zlib"
