Only depend on OpenGLDrv for Android builds
#rnx

diff --git a/Engine/Plugins/Runtime/AR/Google/GoogleARCore/Source/GoogleARCoreBase/GoogleARCoreBase.Build.cs b/Engine/Plugins/Runtime/AR/Google/GoogleARCore/Source/GoogleARCoreBase/GoogleARCoreBase.Build.cs
--- a/Engine/Plugins/Runtime/AR/Google/GoogleARCore/Source/GoogleARCoreBase/GoogleARCoreBase.Build.cs
+++ b/Engine/Plugins/Runtime/AR/Google/GoogleARCore/Source/GoogleARCoreBase/GoogleARCoreBase.Build.cs
@@ -42,7 +42,6 @@ namespace UnrealBuildTool.Rules
 					"GoogleARCoreRendering",
 					"GoogleARCoreSDK",
 					"OpenGL",
-					"OpenGLDrv",
 					"ProceduralMeshComponent",
 					"UElibPNG",
 					"zlib"
@@ -59,11 +58,17 @@ namespace UnrealBuildTool.Rules
 			if (Target.Platform == UnrealTargetPlatform.Android)
 			{
 				// Additional dependencies on android...
-				PrivateDependencyModuleNames.Add("Launch");
+				PrivateDependencyModuleNames.AddRange(
+					new string[]
+					{
+						"Launch",
+						"OpenGLDrv"
+					}
+				);
 
-                // Register Plugin Language
-                string PluginPath = Utils.MakePathRelativeTo(ModuleDirectory, Target.RelativeEnginePath);
-                AdditionalPropertiesForReceipt.Add("AndroidPlugin", Path.Combine(PluginPath, "GoogleARCoreBase_APL.xml"));
+				// Register Plugin Language
+				string PluginPath = Utils.MakePathRelativeTo(ModuleDirectory, Target.RelativeEnginePath);
+				AdditionalPropertiesForReceipt.Add("AndroidPlugin", Path.Combine(PluginPath, "GoogleARCoreBase_APL.xml"));
 				
 				// Needed for including "AndroidEGL.h"
 				string EnginePath = Path.GetFullPath(Target.RelativeEnginePath);
