Switch camera blitter to use bilinear filtering instead of nearest neighbor. Fix by [REDACTED_PERSON]

diff --git a/Engine/Plugins/Runtime/AR/Google/GoogleARCore/Source/GoogleARCoreBase/Private/GoogleARCoreCameraImageBlitter.cpp b/Engine/Plugins/Runtime/AR/Google/GoogleARCore/Source/GoogleARCoreBase/Private/GoogleARCoreCameraImageBlitter.cpp
--- a/Engine/Plugins/Runtime/AR/Google/GoogleARCore/Source/GoogleARCoreBase/Private/GoogleARCoreCameraImageBlitter.cpp
+++ b/Engine/Plugins/Runtime/AR/Google/GoogleARCore/Source/GoogleARCoreBase/Private/GoogleARCoreCameraImageBlitter.cpp
@@ -130,7 +130,7 @@ void FGoogleARCoreDeviceCameraBlitter::LateInit(FIntPoint ImageSize)
 			auto CameraCopy = UTexture2D::CreateTransient(CameraCopySize.X, CameraCopySize.Y, PF_R8G8B8A8);
 			check(CameraCopy);
 			CameraCopy->AddToRoot();
-			CameraCopy->Filter = TextureFilter::TF_Nearest;
+			CameraCopy->Filter = TextureFilter::TF_Bilinear;
 			// The camera texture is in SRGB space, so we need to set this flag to true.
 			// Note that in GLES3.1, textures have the SRGB flag set to true will have be converted to 
 			// linear space automatically when sampling the texture. And Unreal always apply a gamma correction
