Merge 11440643 from Dev-VR: Implement "OnGetSupportedVideoFormats" for ARCore

#rnx

diff --git a/Engine/Plugins/Runtime/AR/Google/GoogleARCore/Source/GoogleARCoreBase/Private/GoogleARCoreXRTrackingSystem.cpp b/Engine/Plugins/Runtime/AR/Google/GoogleARCore/Source/GoogleARCoreBase/Private/GoogleARCoreXRTrackingSystem.cpp
--- a/Engine/Plugins/Runtime/AR/Google/GoogleARCore/Source/GoogleARCoreBase/Private/GoogleARCoreXRTrackingSystem.cpp
+++ b/Engine/Plugins/Runtime/AR/Google/GoogleARCore/Source/GoogleARCoreBase/Private/GoogleARCoreXRTrackingSystem.cpp
@@ -319,6 +319,37 @@ void FGoogleARCoreXRTrackingSystem::OnRemovePin(UARPin* PinToRemove)
 	FGoogleARCoreDevice::GetInstance()->RemoveARPin(PinToRemove);
 }
 
+TArray<FARVideoFormat> FGoogleARCoreXRTrackingSystem::OnGetSupportedVideoFormats(EARSessionType SessionType) const
+{
+	if (SessionType == EARSessionType::None)
+	{
+		return {};
+	}
+
+	// We're creating a temp session here as the "current" session from FGoogleARCoreDevice may not be available
+	// or it may use a different session type
+	const bool bUseFrontCamera = SessionType == EARSessionType::Face;
+	TSharedPtr<FGoogleARCoreSession> NewARCoreSession = FGoogleARCoreSession::CreateARCoreSession(bUseFrontCamera);
+	if (NewARCoreSession && NewARCoreSession->GetSessionCreateStatus() == EGoogleARCoreAPIStatus::AR_SUCCESS)
+	{
+		TArray<FGoogleARCoreCameraConfig> SupportedCameraConfig = NewARCoreSession->GetSupportedCameraConfig();
+
+		TArray<FARVideoFormat> VideoFormats;
+
+		for (const FGoogleARCoreCameraConfig& CameraConfig : SupportedCameraConfig)
+		{
+			// FPS is not exposed so we assume it's 30...
+			// Note that we're using CameraTextureResolution rather than CameraImageResolution as the former is relevant
+			// to the camera passthrough rendering
+			VideoFormats.AddUnique({ 30, CameraConfig.CameraTextureResolution.X, CameraConfig.CameraTextureResolution.Y });
+		}
+
+		return VideoFormats;
+	}
+
+	return {};
+}
+
 TArray<FVector> FGoogleARCoreXRTrackingSystem::OnGetPointCloud() const
 {
 	TArray<FVector> PointCloudPoints;
diff --git a/Engine/Plugins/Runtime/AR/Google/GoogleARCore/Source/GoogleARCoreBase/Private/GoogleARCoreXRTrackingSystem.h b/Engine/Plugins/Runtime/AR/Google/GoogleARCore/Source/GoogleARCoreBase/Private/GoogleARCoreXRTrackingSystem.h
--- a/Engine/Plugins/Runtime/AR/Google/GoogleARCore/Source/GoogleARCoreBase/Private/GoogleARCoreXRTrackingSystem.h
+++ b/Engine/Plugins/Runtime/AR/Google/GoogleARCore/Source/GoogleARCoreBase/Private/GoogleARCoreXRTrackingSystem.h
@@ -73,7 +73,7 @@ protected:
 	virtual TSharedPtr<FARSaveWorldAsyncTask, ESPMode::ThreadSafe> OnSaveWorld() const override { return TSharedPtr<FARSaveWorldAsyncTask, ESPMode::ThreadSafe>(); }
 // @todo -- support these properly
 	virtual EARWorldMappingState OnGetWorldMappingStatus() const override { return EARWorldMappingState::StillMappingNotRelocalizable; }
-	virtual TArray<FARVideoFormat> OnGetSupportedVideoFormats(EARSessionType SessionType) const override { return TArray<FARVideoFormat>(); }
+	virtual TArray<FARVideoFormat> OnGetSupportedVideoFormats(EARSessionType SessionType) const override;
 	virtual TArray<FVector> OnGetPointCloud() const override;
 	//~IARSystemSupport
 
diff --git a/Engine/Source/Runtime/AugmentedReality/Public/ARTypes.h b/Engine/Source/Runtime/AugmentedReality/Public/ARTypes.h
--- a/Engine/Source/Runtime/AugmentedReality/Public/ARTypes.h
+++ b/Engine/Source/Runtime/AugmentedReality/Public/ARTypes.h
@@ -446,6 +446,11 @@ public:
 	{
 		
 	}
+
+	bool operator==(const FARVideoFormat& Other) const
+	{
+		return FPS == Other.FPS && Width == Other.Width && Height == Other.Height;
+	}
 	
 	/** The desired or supported number of frames per second for this video format */
 	UPROPERTY(EditAnywhere, BlueprintReadOnly, Category="AR AugmentedReality|Session")
