Fix memory overwrite in ARCore: The 4th param of ArFrame_transformCoordinates2d is the number of 2D vertices, not the number of floats. So it should be 4 in this case, not 8.
Fix found by [REDACTED_PERSON]

diff --git a/Engine/Plugins/Runtime/AR/Google/GoogleARCore/Source/GoogleARCoreBase/Private/GoogleARCoreAPI.cpp b/Engine/Plugins/Runtime/AR/Google/GoogleARCore/Source/GoogleARCoreBase/Private/GoogleARCoreAPI.cpp
--- a/Engine/Plugins/Runtime/AR/Google/GoogleARCore/Source/GoogleARCoreBase/Private/GoogleARCoreAPI.cpp
+++ b/Engine/Plugins/Runtime/AR/Google/GoogleARCore/Source/GoogleARCoreBase/Private/GoogleARCoreAPI.cpp
@@ -1151,7 +1151,7 @@ void FGoogleARCoreFrame::TransformDisplayUvCoords(const TArray<float>& UvCoords,
 	}
 
 	OutUvCoords.SetNumZeroed(8);
-	ArFrame_transformCoordinates2d(SessionHandle, FrameHandle, AR_COORDINATES_2D_VIEW_NORMALIZED,  8, UvCoords.GetData(), AR_COORDINATES_2D_TEXTURE_NORMALIZED, OutUvCoords.GetData());
+	ArFrame_transformCoordinates2d(SessionHandle, FrameHandle, AR_COORDINATES_2D_VIEW_NORMALIZED, 4, UvCoords.GetData(), AR_COORDINATES_2D_TEXTURE_NORMALIZED, OutUvCoords.GetData());
 #endif
 }
 
