Build fixes.

diff --git a/Engine/Plugins/Experimental/RemoteSession/Source/RemoteSession/Private/Channels/RemoteSessionARCameraChannel.cpp b/Engine/Plugins/Experimental/RemoteSession/Source/RemoteSession/Private/Channels/RemoteSessionARCameraChannel.cpp
--- a/Engine/Plugins/Experimental/RemoteSession/Source/RemoteSession/Private/Channels/RemoteSessionARCameraChannel.cpp
+++ b/Engine/Plugins/Experimental/RemoteSession/Source/RemoteSession/Private/Channels/RemoteSessionARCameraChannel.cpp
@@ -26,7 +26,6 @@
 #include "Containers/DynamicRHIResourceArray.h"
 #include "PostProcess/SceneFilterRendering.h"
 #include "PostProcess/PostProcessMaterial.h"
-#include "PostProcessParameters.h"
 #include "EngineModule.h"
 
 #include "GeneralProjectSettings.h"
diff --git a/Engine/Plugins/Runtime/AR/Google/GoogleARCore/Source/GoogleARCoreRendering/Private/GoogleARCorePassthroughCameraRenderer.cpp b/Engine/Plugins/Runtime/AR/Google/GoogleARCore/Source/GoogleARCoreRendering/Private/GoogleARCorePassthroughCameraRenderer.cpp
--- a/Engine/Plugins/Runtime/AR/Google/GoogleARCore/Source/GoogleARCoreRendering/Private/GoogleARCorePassthroughCameraRenderer.cpp
+++ b/Engine/Plugins/Runtime/AR/Google/GoogleARCore/Source/GoogleARCoreRendering/Private/GoogleARCorePassthroughCameraRenderer.cpp
@@ -13,7 +13,6 @@
 #include "Engine/TextureRenderTarget2D.h"
 #include "PostProcess/SceneFilterRendering.h"
 #include "PostProcess/PostProcessMaterial.h"
-#include "PostProcessParameters.h"
 #include "MaterialShader.h"
 #include "MaterialShaderType.h"
 #include "GoogleARCoreAndroidHelper.h"
diff --git a/Engine/Source/Runtime/Renderer/Private/PostProcess/PostProcessing.h b/Engine/Source/Runtime/Renderer/Private/PostProcess/PostProcessing.h
--- a/Engine/Source/Runtime/Renderer/Private/PostProcess/PostProcessing.h
+++ b/Engine/Source/Runtime/Renderer/Private/PostProcess/PostProcessing.h
@@ -5,6 +5,7 @@
 #include "ScreenPass.h"
 #include "TranslucentRendering.h"
 #include "SystemTextures.h"
+#include "ScenePrivate.h"
 
 class FSceneTextureParameters;
 
diff --git a/Engine/Source/Runtime/Renderer/Private/PostProcess/TemporalAA.cpp b/Engine/Source/Runtime/Renderer/Private/PostProcess/TemporalAA.cpp
--- a/Engine/Source/Runtime/Renderer/Private/PostProcess/TemporalAA.cpp
+++ b/Engine/Source/Runtime/Renderer/Private/PostProcess/TemporalAA.cpp
@@ -8,6 +8,7 @@
 #include "PostProcessing.h"
 #include "SceneTextureParameters.h"
 #include "PixelShaderUtils.h"
+#include "ScenePrivate.h"
 
 namespace
 {
