Fix type of forward declare
#rnx

diff --git a/Engine/Plugins/Runtime/AR/Google/GoogleARCore/Source/GoogleARCoreRendering/Public/GoogleARCorePassthroughCameraRenderer.h b/Engine/Plugins/Runtime/AR/Google/GoogleARCore/Source/GoogleARCoreRendering/Public/GoogleARCorePassthroughCameraRenderer.h
--- a/Engine/Plugins/Runtime/AR/Google/GoogleARCore/Source/GoogleARCoreRendering/Public/GoogleARCorePassthroughCameraRenderer.h
+++ b/Engine/Plugins/Runtime/AR/Google/GoogleARCore/Source/GoogleARCoreRendering/Public/GoogleARCorePassthroughCameraRenderer.h
@@ -21,7 +21,7 @@
 
 class UTexture;
 class FRHICommandListImmediate;
-struct FSceneViewFamily;
+class FSceneViewFamily;
 
 /** A helper class that is used to load the GoogleARCorePassthroughCameraMaterial from its default object. */
 UCLASS()
