LWC: Miscellaneous double/float type conversion fixes

diff --git a/Engine/Plugins/Experimental/UVEditor/Source/UVEditor/Private/UVEditor2DViewportClient.cpp b/Engine/Plugins/Experimental/UVEditor/Source/UVEditor/Private/UVEditor2DViewportClient.cpp
--- a/Engine/Plugins/Experimental/UVEditor/Source/UVEditor/Private/UVEditor2DViewportClient.cpp
+++ b/Engine/Plugins/Experimental/UVEditor/Source/UVEditor/Private/UVEditor2DViewportClient.cpp
@@ -104,7 +104,7 @@ void FUVEditor2DViewportClient::Draw(const FSceneView* View, FPrimitiveDrawInter
 	if (ZoomFactor < 100000 && ZoomFactor > 1)
 	{
 		// Setup and call grid calling function
-		UE::Geometry::FFrame3f LocalFrame(ViewLoc);
+		UE::Geometry::FFrame3f LocalFrame((FVector3f)ViewLoc);
 		FTransform Transform;
 		TArray<FColor> Colors;
 		Colors.Push(GridMajorColor.ToRGBE());
diff --git a/Engine/Plugins/Runtime/AR/Google/GoogleARCore/Source/GoogleARCoreBase/Private/GoogleARCoreAPI.cpp b/Engine/Plugins/Runtime/AR/Google/GoogleARCore/Source/GoogleARCoreBase/Private/GoogleARCoreAPI.cpp
--- a/Engine/Plugins/Runtime/AR/Google/GoogleARCore/Source/GoogleARCoreBase/Private/GoogleARCoreAPI.cpp
+++ b/Engine/Plugins/Runtime/AR/Google/GoogleARCore/Source/GoogleARCoreBase/Private/GoogleARCoreAPI.cpp
@@ -75,7 +75,7 @@ namespace
 
 		FVector ArPosePosition = ARCorePoseMatrix.GetOrigin();
 		FQuat ArPoseRotation = ARCorePoseMatrix.ToQuat();
-		float ArPoseData[7] = { ArPoseRotation.X, ArPoseRotation.Y, ArPoseRotation.Z, ArPoseRotation.W, ArPosePosition.X, ArPosePosition.Y, ArPosePosition.Z };
+		float ArPoseData[7] = { (float)ArPoseRotation.X, (float)ArPoseRotation.Y, (float)ArPoseRotation.Z, (float)ArPoseRotation.W, (float)ArPosePosition.X, (float)ArPosePosition.Y, (float)ArPosePosition.Z };	// LWC_TODO: Precision loss.
 		ArPose_create(SessionHandle, ArPoseData, OutARPose);
 	}
 
@@ -160,7 +160,7 @@ EARTrackingState ToARTrackingState(ArTrackingState State)
 
 FTransform ARCorePoseToUnrealTransform(ArPose* ArPoseHandle, const ArSession* SessionHandle, float WorldToMeterScale)
 {
-	FMatrix ARCorePoseMatrix;
+	FMatrix44f ARCorePoseMatrix;
 	ArPose_getMatrix(SessionHandle, ArPoseHandle, ARCorePoseMatrix.M[0]);
 	FTransform Result = FTransform(ARCoreToUnrealTransform * ARCorePoseMatrix * ARCoreToUnrealTransformInverse);
 	Result.SetLocation(Result.GetLocation() * WorldToMeterScale);
@@ -1211,9 +1211,9 @@ void FGoogleARCoreFrame::ARLineTrace(const FVector& Start, const FVector& End, E
 	}
 
 	float WorldToMeterScale = Session->GetWorldToMeterScale();
-	FVector StartInARCore = UnrealPositionToARCorePosition(Start, WorldToMeterScale);
-	FVector EndInARCore = UnrealPositionToARCorePosition(End, WorldToMeterScale);
-	FVector DirectionInARCore = (EndInARCore - StartInARCore).GetSafeNormal();
+	FVector3f StartInARCore = (FVector3f)UnrealPositionToARCorePosition(Start, WorldToMeterScale);	// LWC_TODO: Precision loss
+	FVector3f EndInARCore = (FVector3f)UnrealPositionToARCorePosition(End, WorldToMeterScale);
+	FVector3f DirectionInARCore = (EndInARCore - StartInARCore).GetSafeNormal();
 	float RayOrigin[3] = { StartInARCore.X, StartInARCore.Y, StartInARCore.Z };
 	float RayDirection[3] = { DirectionInARCore.X, DirectionInARCore.Y, DirectionInARCore.Z };
 
@@ -1247,7 +1247,7 @@ bool FGoogleARCoreFrame::IsDisplayRotationChanged() const
 
 FMatrix FGoogleARCoreFrame::GetProjectionMatrix() const
 {
-	FMatrix ProjectionMatrix;
+	FMatrix44f ProjectionMatrix;
 
 #if PLATFORM_ANDROID
 	if (SessionHandle == nullptr)
diff --git a/Engine/Plugins/Runtime/AR/Google/GoogleARCoreServices/Source/GoogleARCoreServices/Private/GoogleARCoreUtils.h b/Engine/Plugins/Runtime/AR/Google/GoogleARCoreServices/Source/GoogleARCoreServices/Private/GoogleARCoreUtils.h
--- a/Engine/Plugins/Runtime/AR/Google/GoogleARCoreServices/Source/GoogleARCoreServices/Private/GoogleARCoreUtils.h
+++ b/Engine/Plugins/Runtime/AR/Google/GoogleARCoreServices/Source/GoogleARCoreServices/Private/GoogleARCoreUtils.h
@@ -41,7 +41,7 @@ const FMatrix ARCoreToUnrealTransformInverse = ARCoreToUnrealTransform.InverseFa
 
 FTransform ARCorePoseToUnrealTransform(ArPose* ArPoseHandle, const ArSession* SessionHandle, float WorldToMeterScale)
 {
-	FMatrix ARCorePoseMatrix;
+	FMatrix44f ARCorePoseMatrix;
 	ArPose_getMatrix(SessionHandle, ArPoseHandle, ARCorePoseMatrix.M[0]);
 	FTransform Result = FTransform(ARCoreToUnrealTransform * ARCorePoseMatrix * ARCoreToUnrealTransformInverse);
 	Result.SetLocation(Result.GetLocation() * WorldToMeterScale);
@@ -59,7 +59,7 @@ void UnrealTransformToARCorePose(const FTransform& UnrealTransform, const ArSess
 
 	FVector ArPosePosition = ARCorePoseMatrix.GetOrigin();
 	FQuat ArPoseRotation = ARCorePoseMatrix.ToQuat();
-	float ArPoseData[7] = { ArPoseRotation.X, ArPoseRotation.Y, ArPoseRotation.Z, ArPoseRotation.W, ArPosePosition.X, ArPosePosition.Y, ArPosePosition.Z };
+	float ArPoseData[7] = { (float)ArPoseRotation.X, (float)ArPoseRotation.Y, (float)ArPoseRotation.Z, (float)ArPoseRotation.W, (float)ArPosePosition.X, (float)ArPosePosition.Y, (float)ArPosePosition.Z };	// LWC_TODO: Precision loss
 	ArPose_create(SessionHandle, ArPoseData, OutARPose);
 }
 
diff --git a/Engine/Source/Developer/Datasmith/DatasmithFacade/Private/DatasmithFacadeVariant.cpp b/Engine/Source/Developer/Datasmith/DatasmithFacade/Private/DatasmithFacadeVariant.cpp
--- a/Engine/Source/Developer/Datasmith/DatasmithFacade/Private/DatasmithFacadeVariant.cpp
+++ b/Engine/Source/Developer/Datasmith/DatasmithFacade/Private/DatasmithFacadeVariant.cpp
@@ -124,9 +124,9 @@ void FDatasmithFacadeActorBinding::AddRelativeTransformCapture(
 	FVector Location = Transform.GetTranslation();
 	FQuat Rotation = Transform.GetRotation();
 
-	AddRelativeScaleCapture(Scale.X, Scale.Y, Scale.Z);
-	AddRelativeLocationCapture(Location.X, Location.Y, Location.Z);
-	AddRelativeRotationCapture(Rotation.X, Rotation.Y, Rotation.Z, Rotation.W);
+	AddRelativeScaleCapture((float)Scale.X, (float)Scale.Y, (float)Scale.Z);	// LWC_TODO: Precision loss
+	AddRelativeLocationCapture((float)Location.X, (float)Location.Y, (float)Location.Z);
+	AddRelativeRotationCapture((float)Rotation.X, (float)Rotation.Y, (float)Rotation.Z, (float)Rotation.W);
 }
 
 void FDatasmithFacadeActorBinding::AddVisibilityCapture(
diff --git a/Engine/Source/Programs/UnrealLightmass/Private/Lighting/SampleVolume.cpp b/Engine/Source/Programs/UnrealLightmass/Private/Lighting/SampleVolume.cpp
--- a/Engine/Source/Programs/UnrealLightmass/Private/Lighting/SampleVolume.cpp
+++ b/Engine/Source/Programs/UnrealLightmass/Private/Lighting/SampleVolume.cpp
@@ -462,7 +462,7 @@ void FStaticLightingSystem::BeginCalculateVolumeSamples()
 		// Clamp the number of volume samples generated to DynamicObjectSettings.MaxVolumeSamples if necessary by resizing EffectiveVolumeSpacing
 		if (RequestedVolumeSamples > DynamicObjectSettings.MaxVolumeSamples)
 		{
-			EffectiveVolumeSpacing = FMath::Pow(8.0f * VolumeBounds.BoxExtent.X * VolumeBounds.BoxExtent.Y * VolumeBounds.BoxExtent.Z / DynamicObjectSettings.MaxVolumeSamples, .3333333f);
+			EffectiveVolumeSpacing = FMath::Pow(8.0f * float(VolumeBounds.BoxExtent.X * VolumeBounds.BoxExtent.Y * VolumeBounds.BoxExtent.Z / DynamicObjectSettings.MaxVolumeSamples), .3333333f);		// LWC_TODO: Precision loss - float cast
 		}
 			
 		int32 NumUniformVolumeSamples = 0;
