Fixed: ARDebugCheck package crashes on Android when Start AR is hit

diff --git a/Engine/Plugins/Runtime/AR/Google/GoogleARCore/Source/GoogleARCoreRendering/Private/GoogleARCorePassthroughCameraRenderer.cpp b/Engine/Plugins/Runtime/AR/Google/GoogleARCore/Source/GoogleARCoreRendering/Private/GoogleARCorePassthroughCameraRenderer.cpp
--- a/Engine/Plugins/Runtime/AR/Google/GoogleARCore/Source/GoogleARCoreRendering/Private/GoogleARCorePassthroughCameraRenderer.cpp
+++ b/Engine/Plugins/Runtime/AR/Google/GoogleARCore/Source/GoogleARCoreRendering/Private/GoogleARCorePassthroughCameraRenderer.cpp
@@ -197,6 +197,7 @@ void FGoogleARCorePassthroughCameraRenderer::RenderVideoOverlayWithMaterial(FRHI
 		TShaderRef<FGoogleARCoreCameraOverlayVS> VertexShader = MaterialShaderMap->GetShader<FGoogleARCoreCameraOverlayVS>();
 
 		FGraphicsPipelineStateInitializer GraphicsPSOInit;
+		RHICmdList.ApplyCachedRenderTargets(GraphicsPSOInit);
 
 		GraphicsPSOInit.RasterizerState = TStaticRasterizerState<>::GetRHI();
 		
