UE4 -> UE reference in comment.

diff --git a/Engine/Plugins/Runtime/AR/Google/GoogleARCore/Source/GoogleARCoreBase/Public/GoogleARCoreTypes.h b/Engine/Plugins/Runtime/AR/Google/GoogleARCore/Source/GoogleARCoreBase/Public/GoogleARCoreTypes.h
--- a/Engine/Plugins/Runtime/AR/Google/GoogleARCore/Source/GoogleARCoreBase/Public/GoogleARCoreTypes.h
+++ b/Engine/Plugins/Runtime/AR/Google/GoogleARCore/Source/GoogleARCoreBase/Public/GoogleARCoreTypes.h
@@ -126,7 +126,7 @@ enum class EGoogleARCoreCoordinates2DType : uint8
 	Texture = 0,
 	/** ARCore normalized pass-through camera image(for CPU usage) space with top-left (0.0f, 0.0f) and bottom-right(1.0f, 1.0f). */
 	Image = 1,
-	/** UE4 normalized viewport space with top-left(0,0f ,0.0f) and bottom-right(1.0f, 1.0f) */
+	/** UE normalized viewport space with top-left(0,0f ,0.0f) and bottom-right(1.0f, 1.0f) */
 	Viewport = 2
 };
 
