Fix FMatrix cast on Android

diff --git a/Engine/Plugins/Runtime/AR/Google/GoogleARCore/Source/GoogleARCoreBase/Private/GoogleARCoreAPI.cpp b/Engine/Plugins/Runtime/AR/Google/GoogleARCore/Source/GoogleARCoreBase/Private/GoogleARCoreAPI.cpp
--- a/Engine/Plugins/Runtime/AR/Google/GoogleARCore/Source/GoogleARCoreBase/Private/GoogleARCoreAPI.cpp
+++ b/Engine/Plugins/Runtime/AR/Google/GoogleARCore/Source/GoogleARCoreBase/Private/GoogleARCoreAPI.cpp
@@ -162,7 +162,7 @@ FTransform ARCorePoseToUnrealTransform(ArPose* ArPoseHandle, const ArSession* Se
 {
 	FMatrix44f ARCorePoseMatrix;
 	ArPose_getMatrix(SessionHandle, ArPoseHandle, ARCorePoseMatrix.M[0]);
-	FTransform Result = FTransform(ARCoreToUnrealTransform * ARCorePoseMatrix * ARCoreToUnrealTransformInverse);
+	FTransform Result = FTransform(ARCoreToUnrealTransform * FMatrix(ARCorePoseMatrix) * ARCoreToUnrealTransformInverse);
 	Result.SetLocation(Result.GetLocation() * WorldToMeterScale);
 
 	return Result;
@@ -1255,7 +1255,7 @@ FMatrix FGoogleARCoreFrame::GetProjectionMatrix() const
 #if PLATFORM_ANDROID
 	if (SessionHandle == nullptr)
 	{
-		return ProjectionMatrix;
+		return FMatrix();
 	}
 
 	ArCamera_getProjectionMatrix(SessionHandle, CameraHandle, GNearClippingPlane, 100.0f, ProjectionMatrix.M[0]);
diff --git a/Engine/Plugins/Runtime/AR/Google/GoogleARCoreServices/Source/GoogleARCoreServices/Private/GoogleARCoreUtils.h b/Engine/Plugins/Runtime/AR/Google/GoogleARCoreServices/Source/GoogleARCoreServices/Private/GoogleARCoreUtils.h
--- a/Engine/Plugins/Runtime/AR/Google/GoogleARCoreServices/Source/GoogleARCoreServices/Private/GoogleARCoreUtils.h
+++ b/Engine/Plugins/Runtime/AR/Google/GoogleARCoreServices/Source/GoogleARCoreServices/Private/GoogleARCoreUtils.h
@@ -43,7 +43,7 @@ FTransform ARCorePoseToUnrealTransform(ArPose* ArPoseHandle, const ArSession* Se
 {
 	FMatrix44f ARCorePoseMatrix;
 	ArPose_getMatrix(SessionHandle, ArPoseHandle, ARCorePoseMatrix.M[0]);
-	FTransform Result = FTransform(ARCoreToUnrealTransform * ARCorePoseMatrix * ARCoreToUnrealTransformInverse);
+	FTransform Result = FTransform(ARCoreToUnrealTransform * FMatrix(ARCorePoseMatrix) * ARCoreToUnrealTransformInverse);
 	Result.SetLocation(Result.GetLocation() * WorldToMeterScale);
 
 	return Result;
