Adding missing includes.

diff --git a/Engine/Plugins/Runtime/AR/Google/GoogleARCore/Source/GoogleARCoreRendering/Private/GoogleARCorePassthroughCameraRenderer.cpp b/Engine/Plugins/Runtime/AR/Google/GoogleARCore/Source/GoogleARCoreRendering/Private/GoogleARCorePassthroughCameraRenderer.cpp
--- a/Engine/Plugins/Runtime/AR/Google/GoogleARCore/Source/GoogleARCoreRendering/Private/GoogleARCorePassthroughCameraRenderer.cpp
+++ b/Engine/Plugins/Runtime/AR/Google/GoogleARCore/Source/GoogleARCoreRendering/Private/GoogleARCorePassthroughCameraRenderer.cpp
@@ -14,6 +14,7 @@
 #include "Engine/TextureRenderTarget2D.h"
 #include "PostProcess/SceneFilterRendering.h"
 #include "PostProcess/PostProcessMaterial.h"
+#include "Materials/MaterialRenderProxy.h"
 #include "MaterialShader.h"
 #include "MaterialShaderType.h"
 #include "GoogleARCoreAndroidHelper.h"
diff --git a/Engine/Source/Runtime/Engine/Private/Materials/MaterialAttributeDefinitionMap.cpp b/Engine/Source/Runtime/Engine/Private/Materials/MaterialAttributeDefinitionMap.cpp
--- a/Engine/Source/Runtime/Engine/Private/Materials/MaterialAttributeDefinitionMap.cpp
+++ b/Engine/Source/Runtime/Engine/Private/Materials/MaterialAttributeDefinitionMap.cpp
@@ -9,6 +9,7 @@
 #include "Materials/MaterialExpressionShadingModel.h"
 #include "Materials/MaterialExpressionStrata.h"
 #include "SceneTypes.h"
+#include "StrataDefinitions.h"
 
 #define LOCTEXT_NAMESPACE "MaterialShared"
 
