UE-173464 Ensure in AR Template on Android
-Added implementation of GetReferencerName() as required by FGCObject refactor.

diff --git a/Engine/Plugins/Runtime/AR/Google/GoogleARCore/Source/GoogleARCoreBase/Private/GoogleARCoreXRTrackingSystem.h b/Engine/Plugins/Runtime/AR/Google/GoogleARCore/Source/GoogleARCoreBase/Private/GoogleARCoreXRTrackingSystem.h
--- a/Engine/Plugins/Runtime/AR/Google/GoogleARCore/Source/GoogleARCoreBase/Private/GoogleARCoreXRTrackingSystem.h
+++ b/Engine/Plugins/Runtime/AR/Google/GoogleARCore/Source/GoogleARCoreBase/Private/GoogleARCoreXRTrackingSystem.h
@@ -98,6 +98,10 @@ protected:
 private:
 	//~ FGCObject
 	virtual void AddReferencedObjects(FReferenceCollector& Collector) override;
+	virtual FString GetReferencerName() const override
+	{
+		return TEXT("FGoogleARCoreXRTrackingSystem");
+	}
 	//~ FGCObject
 
 private:
