Adding missing includes for NonUnity.

diff --git a/Engine/Plugins/Runtime/AR/Google/GoogleARCore/Source/GoogleARCoreRendering/Private/GoogleARCorePassthroughCameraRenderer.cpp b/Engine/Plugins/Runtime/AR/Google/GoogleARCore/Source/GoogleARCoreRendering/Private/GoogleARCorePassthroughCameraRenderer.cpp
--- a/Engine/Plugins/Runtime/AR/Google/GoogleARCore/Source/GoogleARCoreRendering/Private/GoogleARCorePassthroughCameraRenderer.cpp
+++ b/Engine/Plugins/Runtime/AR/Google/GoogleARCore/Source/GoogleARCoreRendering/Private/GoogleARCorePassthroughCameraRenderer.cpp
@@ -16,12 +16,14 @@
 #include "MaterialDomain.h"
 #include "MaterialShader.h"
 #include "MaterialShaderType.h"
+#include "Materials/MaterialRenderProxy.h"
 #include "GoogleARCoreAndroidHelper.h"
 #include "CommonRenderResources.h"
 #include "ARUtilitiesFunctionLibrary.h"
 #include "HAL/IConsoleManager.h"
 #include "UObject/Package.h"
-
+#include "PostProcess/DrawRectangle.h"
+#include "RHIStaticStates.h"
 
 #if !(UE_BUILD_SHIPPING || UE_BUILD_TEST)
 	#define ALLOWS_DEBUG_OVERLAY 1
