UE-181824 Move HeadMountedDisplay files to XRBase plugin to reduce UE overhead. - fix GoogleARCoreAugmentedFace build error
-Missing include.

diff --git a/Engine/Plugins/Runtime/AR/Google/GoogleARCore/Source/GoogleARCoreBase/Private/GoogleARCoreAugmentedFace.cpp b/Engine/Plugins/Runtime/AR/Google/GoogleARCore/Source/GoogleARCoreBase/Private/GoogleARCoreAugmentedFace.cpp
--- a/Engine/Plugins/Runtime/AR/Google/GoogleARCore/Source/GoogleARCoreBase/Private/GoogleARCoreAugmentedFace.cpp
+++ b/Engine/Plugins/Runtime/AR/Google/GoogleARCore/Source/GoogleARCoreBase/Private/GoogleARCoreAugmentedFace.cpp
@@ -2,7 +2,9 @@
 
 #include "GoogleARCoreAugmentedFace.h"
 
+#include "IXRTrackingSystem.h"
 #include "ARSystem.h"
+#include "ARSupportInterface.h"
 
 void UGoogleARCoreAugmentedFace::DebugDraw(UWorld* World, const FLinearColor& OutlineColor, float OutlineThickness, float PersistForSeconds /*= 0.0f*/) const
 {
