Suppress deterministic warnings for several plugin modules that use __DATE__ and __TIME__

#rnx

diff --git a/Engine/Plugins/Runtime/AR/Google/GoogleARCore/Source/GoogleARCoreBase/GoogleARCoreBase.Build.cs b/Engine/Plugins/Runtime/AR/Google/GoogleARCore/Source/GoogleARCoreBase/GoogleARCoreBase.Build.cs
--- a/Engine/Plugins/Runtime/AR/Google/GoogleARCore/Source/GoogleARCoreBase/GoogleARCoreBase.Build.cs
+++ b/Engine/Plugins/Runtime/AR/Google/GoogleARCore/Source/GoogleARCoreBase/GoogleARCoreBase.Build.cs
@@ -10,6 +10,8 @@ namespace UnrealBuildTool.Rules
 	{
 		public GoogleARCoreBase(ReadOnlyTargetRules Target) : base(Target)
 		{
+			DeterministicWarningLevel = WarningLevel.Off; // __DATE__ __TIME__ in Private/GoogleARCoreXRTrackingSystem.cpp
+
 			PublicIncludePathModuleNames.Add("TargetPlatform");
 
 			PublicDependencyModuleNames.AddRange(new string[]
diff --git a/Engine/Source/Runtime/Datasmith/CADKernel/CADKernel.Build.cs b/Engine/Source/Runtime/Datasmith/CADKernel/CADKernel.Build.cs
--- a/Engine/Source/Runtime/Datasmith/CADKernel/CADKernel.Build.cs
+++ b/Engine/Source/Runtime/Datasmith/CADKernel/CADKernel.Build.cs
@@ -8,6 +8,8 @@ namespace UnrealBuildTool.Rules
 		public CADKernel(ReadOnlyTargetRules Target)
 			: base(Target)
 		{
+			DeterministicWarningLevel = WarningLevel.Off; // __DATE__ in Private/CADKernel/Core/System.cpp
+
 			PublicDependencyModuleNames.AddRange(
 				new string[]
 				{
