UE-127308 ARCore: Deprecate UGoogleARCoreFrameFunctionLibrary::GetCameraMetadata
UE-130223 Several ARCore features that have been deprecated still have references in 4.26
-Deprecate GetCameraMetadata.
-Delete a large number of long deprecated functions.
-Removed blueprint access to a number of enums and structs that were only used by those long deprecated functions.

diff --git a/Engine/Plugins/Runtime/AR/Google/GoogleARCore/Source/GoogleARCoreBase/Private/GoogleARCoreFunctionLibrary.cpp b/Engine/Plugins/Runtime/AR/Google/GoogleARCore/Source/GoogleARCoreBase/Private/GoogleARCoreFunctionLibrary.cpp
--- a/Engine/Plugins/Runtime/AR/Google/GoogleARCore/Source/GoogleARCoreBase/Private/GoogleARCoreFunctionLibrary.cpp
+++ b/Engine/Plugins/Runtime/AR/Google/GoogleARCore/Source/GoogleARCoreBase/Private/GoogleARCoreFunctionLibrary.cpp
@@ -3,7 +3,6 @@
 #include "GoogleARCoreFunctionLibrary.h"
 #include "UnrealEngine.h"
 #include "Engine/Engine.h"
-#include "LatentActions.h"
 #include "Misc/AssertionMacros.h"
 #include "ARBlueprintLibrary.h"
 
@@ -12,290 +11,6 @@
 #include "GoogleARCoreDevice.h"
 #include "GoogleARCoreXRTrackingSystem.h"
 
-namespace
-{
-	EGoogleARCoreInstallRequestResult ToAPKInstallStatus(EGoogleARCoreAPIStatus RequestStatus)
-	{
-		EGoogleARCoreInstallRequestResult OutRequestResult = EGoogleARCoreInstallRequestResult::FatalError;
-		switch (RequestStatus)
-		{
-			case EGoogleARCoreAPIStatus::AR_SUCCESS:
-				OutRequestResult = EGoogleARCoreInstallRequestResult::Installed;
-				break;
-			case EGoogleARCoreAPIStatus::AR_ERROR_FATAL:
-				OutRequestResult = EGoogleARCoreInstallRequestResult::FatalError;
-				break;
-			case EGoogleARCoreAPIStatus::AR_UNAVAILABLE_DEVICE_NOT_COMPATIBLE:
-				OutRequestResult = EGoogleARCoreInstallRequestResult::DeviceNotCompatible;
-				break;
-			case EGoogleARCoreAPIStatus::AR_UNAVAILABLE_USER_DECLINED_INSTALLATION:
-				OutRequestResult = EGoogleARCoreInstallRequestResult::UserDeclinedInstallation;
-				break;
-			default:
-				ensureMsgf(false, TEXT("Unexpected ARCore API Status: %d"), static_cast<int>(RequestStatus));
-				break;
-		}
-		return OutRequestResult;
-	}
-
-	const float DefaultLineTraceDistance = 100000; // 1000 meter
-}
-
-/************************************************************************/
-/*  UGoogleARCoreSessionFunctionLibrary | Lifecycle                     */
-/************************************************************************/
-struct FARCoreCheckAvailabilityAction : public FPendingLatentAction
-{
-public:
-	FName ExecutionFunction;
-	int32 OutputLink;
-	FWeakObjectPtr CallbackTarget;
-	EGoogleARCoreAvailability& OutAvailability;
-
-	FARCoreCheckAvailabilityAction(const FLatentActionInfo& InLatentInfo, EGoogleARCoreAvailability& InAvailability)
-	: FPendingLatentAction()
-	, ExecutionFunction(InLatentInfo.ExecutionFunction)
-	, OutputLink(InLatentInfo.Linkage)
-	, CallbackTarget(InLatentInfo.CallbackTarget)
-	, OutAvailability(InAvailability)
-	{}
-
-	virtual void UpdateOperation(FLatentResponse& Response) override
-	{
-		EGoogleARCoreAvailability ARCoreAvailability = FGoogleARCoreDevice::GetInstance()->CheckARCoreAPKAvailability();
-		if (ARCoreAvailability != EGoogleARCoreAvailability::UnknownChecking)
-		{
-			OutAvailability = ARCoreAvailability;
-			Response.FinishAndTriggerIf(true, ExecutionFunction, OutputLink, CallbackTarget);
-		}
-	}
-#if WITH_EDITOR
-	virtual FString GetDescription() const override
-	{
-		return FString::Printf(TEXT("Checking ARCore availability."));
-	}
-#endif
-};
-
-void UGoogleARCoreSessionFunctionLibrary::CheckARCoreAvailability(UObject* WorldContextObject, struct FLatentActionInfo LatentInfo, EGoogleARCoreAvailability& OutAvailability)
-{
-	if (UWorld* World = GEngine->GetWorldFromContextObject(WorldContextObject, EGetWorldErrorMode::LogAndReturnNull))
-	{
-		FLatentActionManager& LatentManager = World->GetLatentActionManager();
-		if (LatentManager.FindExistingAction<FARCoreCheckAvailabilityAction>(LatentInfo.CallbackTarget, LatentInfo.UUID) == nullptr)
-		{
-			FARCoreCheckAvailabilityAction* NewAction = new FARCoreCheckAvailabilityAction(LatentInfo, OutAvailability);
-			LatentManager.AddNewAction(LatentInfo.CallbackTarget, LatentInfo.UUID, NewAction);
-		}
-		else
-		{
-			UE_LOG(LogGoogleARCore, Warning, TEXT("CheckARCoreAvailability not excuted. The previous action hasn't finished yet."));
-		}
-	}
-}
-
-EGoogleARCoreAvailability UGoogleARCoreSessionFunctionLibrary::CheckARCoreAvailableStatus()
-{
-	return FGoogleARCoreDevice::GetInstance()->CheckARCoreAPKAvailability();
-}
-
-struct FARCoreAPKInstallAction : public FPendingLatentAction
-{
-public:
-	FName ExecutionFunction;
-	int32 OutputLink;
-	FWeakObjectPtr CallbackTarget;
-	EGoogleARCoreInstallRequestResult& OutRequestResult;
-	bool bInstallRequested;
-
-	FARCoreAPKInstallAction(const FLatentActionInfo& InLatentInfo, EGoogleARCoreInstallRequestResult& InRequestResult)
-	: FPendingLatentAction()
-	, ExecutionFunction(InLatentInfo.ExecutionFunction)
-	, OutputLink(InLatentInfo.Linkage)
-	, CallbackTarget(InLatentInfo.CallbackTarget)
-	, OutRequestResult(InRequestResult)
-	, bInstallRequested(false)
-	{}
-
-	virtual void UpdateOperation(FLatentResponse& Response) override
-	{
-		EGoogleARCoreInstallStatus InstallStatus = EGoogleARCoreInstallStatus::Installed;
-		EGoogleARCoreAPIStatus RequestStatus = FGoogleARCoreDevice::GetInstance()->RequestInstall(!bInstallRequested, InstallStatus);
-		if (RequestStatus != EGoogleARCoreAPIStatus::AR_SUCCESS)
-		{
-			OutRequestResult = ToAPKInstallStatus(RequestStatus);
-			Response.FinishAndTriggerIf(true, ExecutionFunction, OutputLink, CallbackTarget);
-		}
-		else if (InstallStatus == EGoogleARCoreInstallStatus::Installed)
-		{
-			OutRequestResult = EGoogleARCoreInstallRequestResult::Installed;
-			Response.FinishAndTriggerIf(true, ExecutionFunction, OutputLink, CallbackTarget);
-		}
-		else
-		{
-			// InstallSttatus returns requested.
-			bInstallRequested = true;
-		}
-	}
-#if WITH_EDITOR
-	virtual FString GetDescription() const override
-	{
-		return FString::Printf(TEXT("Checking ARCore availability."));
-	}
-#endif
-};
-
-void UGoogleARCoreSessionFunctionLibrary::InstallARCoreService(UObject* WorldContextObject, struct FLatentActionInfo LatentInfo, EGoogleARCoreInstallRequestResult& OutInstallResult)
-{
-	if (UWorld* World = GEngine->GetWorldFromContextObject(WorldContextObject, EGetWorldErrorMode::LogAndReturnNull))
-	{
-		FLatentActionManager& LatentManager = World->GetLatentActionManager();
-		if (LatentManager.FindExistingAction<FARCoreAPKInstallAction>(LatentInfo.CallbackTarget, LatentInfo.UUID) == nullptr)
-		{
-			FARCoreAPKInstallAction* NewAction = new FARCoreAPKInstallAction(LatentInfo, OutInstallResult);
-			LatentManager.AddNewAction(LatentInfo.CallbackTarget, LatentInfo.UUID, NewAction);
-		}
-	}
-}
-
-EGoogleARCoreInstallStatus UGoogleARCoreSessionFunctionLibrary::RequestInstallARCoreAPK()
-{
-	EGoogleARCoreInstallStatus InstallStatus = EGoogleARCoreInstallStatus::Installed;
-	EGoogleARCoreAPIStatus RequestStatus = FGoogleARCoreDevice::GetInstance()->RequestInstall(true, InstallStatus);
-
-	return InstallStatus;
-}
-
-EGoogleARCoreInstallRequestResult UGoogleARCoreSessionFunctionLibrary::GetARCoreAPKInstallResult()
-{
-	EGoogleARCoreInstallStatus InstallStatus = EGoogleARCoreInstallStatus::Installed;
-	EGoogleARCoreAPIStatus RequestStatus = FGoogleARCoreDevice::GetInstance()->RequestInstall(false, InstallStatus);
-
-	return ToAPKInstallStatus(RequestStatus);
-}
-
-UGoogleARCoreEventManager* UGoogleARCoreSessionFunctionLibrary::GetARCoreEventManager()
-{
-	if (auto TrackingSystem = FGoogleARCoreXRTrackingSystem::GetInstance())
-	{
-		return TrackingSystem->GetEventManager();
-	}
-	else
-	{
-		return nullptr;
-	}
-}
-
-struct FARCoreStartSessionAction : public FPendingLatentAction
-{
-public:
-	FName ExecutionFunction;
-	int32 OutputLink;
-	FWeakObjectPtr CallbackTarget;
-
-	FARCoreStartSessionAction(const FLatentActionInfo& InLatentInfo)
-		: FPendingLatentAction()
-		, ExecutionFunction(InLatentInfo.ExecutionFunction)
-		, OutputLink(InLatentInfo.Linkage)
-		, CallbackTarget(InLatentInfo.CallbackTarget)
-	{}
-
-	virtual void UpdateOperation(FLatentResponse& Response) override
-	{
-		bool bSessionStartedFinished = FGoogleARCoreDevice::GetInstance()->GetStartSessionRequestFinished();
-		Response.FinishAndTriggerIf(bSessionStartedFinished, ExecutionFunction, OutputLink, CallbackTarget);
-	}
-#if WITH_EDITOR
-	virtual FString GetDescription() const override
-	{
-		return FString::Printf(TEXT("Starting ARCore tracking session"));
-	}
-#endif
-};
-
-void UGoogleARCoreSessionFunctionLibrary::StartARCoreSession(UObject* WorldContextObject, struct FLatentActionInfo LatentInfo, UGoogleARCoreSessionConfig* Configuration)
-{
-	if (UWorld* World = GEngine->GetWorldFromContextObject(WorldContextObject, EGetWorldErrorMode::LogAndReturnNull))
-	{
-		FLatentActionManager& LatentManager = World->GetLatentActionManager();
-		if (LatentManager.FindExistingAction<FARCoreStartSessionAction>(LatentInfo.CallbackTarget, LatentInfo.UUID) == nullptr)
-		{
-			UARBlueprintLibrary::StartARSession(static_cast<UARSessionConfig*>(Configuration));
-			FARCoreStartSessionAction* NewAction = new FARCoreStartSessionAction(LatentInfo);
-			LatentManager.AddNewAction(LatentInfo.CallbackTarget, LatentInfo.UUID, NewAction);
-		}
-	}
-}
-
-bool UGoogleARCoreSessionFunctionLibrary::SetARCoreCameraConfig(FGoogleARCoreCameraConfig TargetCameraConfig)
-{
-	return FGoogleARCoreDevice::GetInstance()->SetARCameraConfig(TargetCameraConfig);
-}
-
-bool UGoogleARCoreSessionFunctionLibrary::GetARCoreCameraConfig(FGoogleARCoreCameraConfig& OutCurrentCameraConfig)
-{
-	return FGoogleARCoreDevice::GetInstance()->GetARCameraConfig(OutCurrentCameraConfig);
-}
-
-/************************************************************************/
-/*  UGoogleARCoreSessionFunctionLibrary | PassthroughCamera             */
-/************************************************************************/
-bool UGoogleARCoreSessionFunctionLibrary::IsPassthroughCameraRenderingEnabled()
-{
-	if (auto TrackingSystem = FGoogleARCoreXRTrackingSystem::GetInstance())
-	{
-		return TrackingSystem->GetColorCameraRenderingEnabled();
-	}
-	else
-	{
-		UE_LOG(LogGoogleARCore, Error, TEXT("Failed to find GoogleARCoreXRTrackingSystem: GoogleARCoreXRTrackingSystem is not available."));
-	}
-	return false;
-}
-
-void UGoogleARCoreSessionFunctionLibrary::SetPassthroughCameraRenderingEnabled(bool bEnable)
-{
-	if (auto TrackingSystem = FGoogleARCoreXRTrackingSystem::GetInstance())
-	{
-		TrackingSystem->EnableColorCameraRendering(bEnable);
-	}
-	else
-	{
-		UE_LOG(LogGoogleARCore, Error, TEXT("Failed to config ARCoreXRCamera: GoogleARCoreXRTrackingSystem is not available."));
-	}
-}
-
-void UGoogleARCoreSessionFunctionLibrary::GetPassthroughCameraImageUV(const TArray<float>& InUV, TArray<float>& OutUV)
-{
-	FGoogleARCoreDevice::GetInstance()->GetPassthroughCameraImageUVs(InUV, OutUV);
-}
-
-void UGoogleARCoreSessionFunctionLibrary::GetAllPlanes(TArray<UARPlaneGeometry*>& OutPlaneList)
-{
-	FGoogleARCoreDevice::GetInstance()->GetAllTrackables<UARPlaneGeometry>(OutPlaneList);
-}
-
-void UGoogleARCoreSessionFunctionLibrary::GetAllTrackablePoints(TArray<UARTrackedPoint*>& OutTrackablePointList)
-{
-	FGoogleARCoreDevice::GetInstance()->GetAllTrackables<UARTrackedPoint>(OutTrackablePointList);
-}
-
-void UGoogleARCoreSessionFunctionLibrary::GetAllAugmentedImages(TArray<UGoogleARCoreAugmentedImage*>& OutAugmentedImageList)
-{
-	FGoogleARCoreDevice::GetInstance()->GetAllTrackables<UGoogleARCoreAugmentedImage>(OutAugmentedImageList);
-}
-
-void UGoogleARCoreSessionFunctionLibrary::GetAllAugmentedFaces(TArray<UGoogleARCoreAugmentedFace*>& OutAugmentedFaceList)
-{
-	FGoogleARCoreDevice::GetInstance()->GetAllTrackables<UGoogleARCoreAugmentedFace>(OutAugmentedFaceList);
-}
-
-template< class T >
-void UGoogleARCoreSessionFunctionLibrary::GetAllTrackable(TArray<T*>& OutTrackableList)
-{
-	FGoogleARCoreDevice::GetInstance()->GetAllTrackables<T>(OutTrackableList);
-}
-
 UARCandidateImage* UGoogleARCoreSessionFunctionLibrary::AddRuntimeCandidateImageFromRawbytes(UARSessionConfig* SessionConfig, const TArray<uint8>& ImageGrayscalePixels,
 	int ImageWidth, int ImageHeight, FString FriendlyName, float PhysicalWidth, UTexture2D* CandidateTexture /*= nullptr*/)
 {
@@ -313,144 +28,13 @@ UARCandidateImage* UGoogleARCoreSessionFunctionLibrary::AddRuntimeCandidateImage
 	return nullptr;
 }
 
-template void UGoogleARCoreSessionFunctionLibrary::GetAllTrackable<UARTrackedGeometry>(TArray<UARTrackedGeometry*>& OutTrackableList);
-template void UGoogleARCoreSessionFunctionLibrary::GetAllTrackable<UARPlaneGeometry>(TArray<UARPlaneGeometry*>& OutTrackableList);
-template void UGoogleARCoreSessionFunctionLibrary::GetAllTrackable<UARTrackedPoint>(TArray<UARTrackedPoint*>& OutTrackableList);
-
 /************************************************************************/
 /*  UGoogleARCoreFrameFunctionLibrary                                   */
 /************************************************************************/
-EGoogleARCoreTrackingState UGoogleARCoreFrameFunctionLibrary::GetTrackingState()
-{
-	return FGoogleARCoreDevice::GetInstance()->GetTrackingState();
-}
-
-EGoogleARCoreTrackingFailureReason UGoogleARCoreFrameFunctionLibrary::GetTrackingFailureReason()
-{
-	return FGoogleARCoreDevice::GetInstance()->GetTrackingFailureReason();
-}
-
-void UGoogleARCoreFrameFunctionLibrary::GetPose(FTransform& LastePose)
-{
-	LastePose = FGoogleARCoreDevice::GetInstance()->GetLatestPose();
-}
-
-bool UGoogleARCoreFrameFunctionLibrary::ARCoreLineTrace(UObject* WorldContextObject, const FVector2D& ScreenPosition, TSet<EGoogleARCoreLineTraceChannel> TraceChannels, TArray<FARTraceResult>& OutHitResults)
-{
-	EGoogleARCoreLineTraceChannel TraceChannelValue = EGoogleARCoreLineTraceChannel::None;
-	for (EGoogleARCoreLineTraceChannel Channel : TraceChannels)
-	{
-		TraceChannelValue = TraceChannelValue | Channel;
-	}
-
-	FGoogleARCoreDevice::GetInstance()->ARLineTrace(ScreenPosition, TraceChannelValue, OutHitResults);
-	return OutHitResults.Num() > 0;
-}
-
-bool UGoogleARCoreFrameFunctionLibrary::ARCoreLineTraceRay(UObject* WorldContextObject, const FVector& Start, const FVector& End, TSet<EGoogleARCoreLineTraceChannel> TraceChannels, TArray<FARTraceResult>& OutHitResults)
-{
-	EGoogleARCoreLineTraceChannel TraceChannelValue = EGoogleARCoreLineTraceChannel::None;
-	for (EGoogleARCoreLineTraceChannel Channel : TraceChannels)
-	{
-		TraceChannelValue = TraceChannelValue | Channel;
-	}
-
-	FGoogleARCoreDevice::GetInstance()->ARLineTrace(Start, End, TraceChannelValue, OutHitResults);
-	return OutHitResults.Num() > 0;
-}
-
-void UGoogleARCoreFrameFunctionLibrary::GetUpdatedARPins(TArray<UARPin*>& OutAnchorList)
-{
-	FGoogleARCoreDevice::GetInstance()->GetUpdatedARPins(OutAnchorList);
-}
-
-void UGoogleARCoreFrameFunctionLibrary::GetUpdatedPlanes(TArray<UARPlaneGeometry*>& OutPlaneList)
-{
-	FGoogleARCoreDevice::GetInstance()->GetUpdatedTrackables<UARPlaneGeometry>(OutPlaneList);
-}
-
-void UGoogleARCoreFrameFunctionLibrary::GetUpdatedTrackablePoints(TArray<UARTrackedPoint*>& OutTrackablePointList)
-{
-	FGoogleARCoreDevice::GetInstance()->GetUpdatedTrackables<UARTrackedPoint>(OutTrackablePointList);
-}
-
-void UGoogleARCoreFrameFunctionLibrary::GetUpdatedAugmentedImages(TArray<UGoogleARCoreAugmentedImage*>& OutImageList)
-{
-	FGoogleARCoreDevice::GetInstance()->GetUpdatedTrackables<UGoogleARCoreAugmentedImage>(OutImageList);
-}
-
-void UGoogleARCoreFrameFunctionLibrary::GetUpdatedAugmentedFaces(TArray<UGoogleARCoreAugmentedFace*>& OutFaceList)
-{
-	FGoogleARCoreDevice::GetInstance()->GetUpdatedTrackables<UGoogleARCoreAugmentedFace>(OutFaceList);
-}
-
-void UGoogleARCoreFrameFunctionLibrary::GetLightEstimation(FGoogleARCoreLightEstimate& LightEstimation)
-{
-	LightEstimation = FGoogleARCoreDevice::GetInstance()->GetLatestLightEstimate();
-}
-
-EGoogleARCoreFunctionStatus UGoogleARCoreFrameFunctionLibrary::GetPointCloud(UGoogleARCorePointCloud*& OutLatestPointCloud)
-{
-	return FGoogleARCoreDevice::GetInstance()->GetLatestPointCloud(OutLatestPointCloud);
-}
-
-EGoogleARCoreFunctionStatus UGoogleARCoreFrameFunctionLibrary::AcquirePointCloud(UGoogleARCorePointCloud*& OutLatestPointCloud)
-{
-	return FGoogleARCoreDevice::GetInstance()->AcquireLatestPointCloud(OutLatestPointCloud);
-}
 
 #if PLATFORM_ANDROID
 EGoogleARCoreFunctionStatus UGoogleARCoreFrameFunctionLibrary::GetCameraMetadata(const ACameraMetadata*& OutCameraMetadata)
 {
 	return FGoogleARCoreDevice::GetInstance()->GetLatestCameraMetadata(OutCameraMetadata);
 }
-#endif
-
-template< class T >
-void UGoogleARCoreFrameFunctionLibrary::GetUpdatedTrackable(TArray<T*>& OutTrackableList)
-{
-	FGoogleARCoreDevice::GetInstance()->GetUpdatedTrackables<T>(OutTrackableList);
-}
-
-template void UGoogleARCoreFrameFunctionLibrary::GetUpdatedTrackable<UARTrackedGeometry>(TArray<UARTrackedGeometry*>& OutTrackableList);
-template void UGoogleARCoreFrameFunctionLibrary::GetUpdatedTrackable<UARPlaneGeometry>(TArray<UARPlaneGeometry*>& OutTrackableList);
-template void UGoogleARCoreFrameFunctionLibrary::GetUpdatedTrackable<UARTrackedPoint>(TArray<UARTrackedPoint*>& OutTrackableList);
-
-UTexture* UGoogleARCoreFrameFunctionLibrary::GetCameraTexture()
-{
-	return nullptr;
-}
-
-EGoogleARCoreFunctionStatus UGoogleARCoreFrameFunctionLibrary::AcquireCameraImage(UGoogleARCoreCameraImage *&OutLatestCameraImage)
-{
-	return FGoogleARCoreDevice::GetInstance()->AcquireCameraImage(OutLatestCameraImage);
-}
-
-void UGoogleARCoreFrameFunctionLibrary::TransformARCoordinates2D(EGoogleARCoreCoordinates2DType InputCoordinatesType, const TArray<FVector2D>& InputCoordinates, EGoogleARCoreCoordinates2DType OutputCoordinatesType, TArray<FVector2D>& OutputCoordinates)
-{
-	return FGoogleARCoreDevice::GetInstance()->TransformARCoordinates2D(InputCoordinatesType, InputCoordinates, OutputCoordinatesType, OutputCoordinates);
-}
-
-EGoogleARCoreFunctionStatus UGoogleARCoreFrameFunctionLibrary::GetCameraImageIntrinsics(UGoogleARCoreCameraIntrinsics*& OutCameraIntrinsics)
-{
-	FARCameraIntrinsics CameraIntrinsics;
-	const auto Result = FGoogleARCoreDevice::GetInstance()->GetCameraImageIntrinsics(CameraIntrinsics);
-	if (Result == EGoogleARCoreFunctionStatus::Success)
-	{
-		OutCameraIntrinsics = NewObject<UGoogleARCoreCameraIntrinsics>();
-		OutCameraIntrinsics->SetCameraIntrinsics(CameraIntrinsics);
-	}
-	return Result;
-}
-
-EGoogleARCoreFunctionStatus UGoogleARCoreFrameFunctionLibrary::GetCameraTextureIntrinsics(UGoogleARCoreCameraIntrinsics*& OutCameraIntrinsics)
-{
-	FARCameraIntrinsics CameraIntrinsics;
-	const auto Result = FGoogleARCoreDevice::GetInstance()->GetCameraTextureIntrinsics(CameraIntrinsics);
-	if (Result == EGoogleARCoreFunctionStatus::Success)
-	{
-		OutCameraIntrinsics = NewObject<UGoogleARCoreCameraIntrinsics>();
-		OutCameraIntrinsics->SetCameraIntrinsics(CameraIntrinsics);
-	}
-	return Result;
-}
+#endif
\ No newline at end of file
diff --git a/Engine/Plugins/Runtime/AR/Google/GoogleARCore/Source/GoogleARCoreBase/Public/GoogleARCoreFunctionLibrary.h b/Engine/Plugins/Runtime/AR/Google/GoogleARCore/Source/GoogleARCoreBase/Public/GoogleARCoreFunctionLibrary.h
--- a/Engine/Plugins/Runtime/AR/Google/GoogleARCore/Source/GoogleARCoreBase/Public/GoogleARCoreFunctionLibrary.h
+++ b/Engine/Plugins/Runtime/AR/Google/GoogleARCore/Source/GoogleARCoreBase/Public/GoogleARCoreFunctionLibrary.h
@@ -3,13 +3,9 @@
 #pragma once
 
 #include "CoreMinimal.h"
-#include "Engine/LatentActionManager.h"
 #include "Kismet/BlueprintFunctionLibrary.h"
 
 #include "GoogleARCoreTypes.h"
-#include "GoogleARCoreAugmentedFace.h"
-#include "GoogleARCoreAugmentedImage.h"
-#include "GoogleARCoreSessionConfig.h"
 #include "GoogleARCoreFunctionLibrary.generated.h"
 
 /** A function library that provides static/Blueprint functions associated with GoogleARCore session.*/
@@ -19,174 +15,6 @@ class GOOGLEARCOREBASE_API UGoogleARCoreSessionFunctionLibrary : public UBluepri
 	GENERATED_BODY()
 
 public:
-	//-----------------Lifecycle---------------------
-
-	/**
-	 * A Latent Action to check the availability of ARCore on this device.
-	 * This may initiate a query with a remote service to determine if the device is supported by ARCore. The Latent Action will complete when the check is finished.
-	 *
-	 * @param WorldContextObject	The world context.
-	 * @param LatentInfo			Unreal internal type required for all latent actions.
-	 * @param OutAvailability		The availability result as a EGoogleARCoreAvailability.
-	 */
-	UE_DEPRECATED(4.26, "This function is deprecated, use CheckARServiceAvailability from UARDependencyHandler.")
-	UFUNCTION(BlueprintCallable, Category = "GoogleARCore|Availability", meta = (Latent, LatentInfo = "LatentInfo", WorldContext = "WorldContextObject", Keywords = "googlear arcore availability", DeprecatedFunction, DeprecationMessage="This function is deprecated, use CheckARServiceAvailability from UARDependencyHandler."))
-	static void CheckARCoreAvailability(UObject* WorldContextObject, struct FLatentActionInfo LatentInfo, EGoogleARCoreAvailability& OutAvailability);
-
-	/**
-	 * A Latent Action to initiates installation of ARCore if required.
-	 * This function may cause your application be paused if installing ARCore is required.
-	 *
-	 * @param WorldContextObject	The world context.
-	 * @param LatentInfo			Unreal internal type required for all latent actions.
-	 * @param OutInstallResult		The install request result as a EGoogleARCoreInstallRequestResult.
-	 */
-	UE_DEPRECATED(4.26, "This function is deprecated, use InstallARService from UARDependencyHandler.")
-	UFUNCTION(BlueprintCallable, Category = "GoogleARCore|Availability", meta = (Latent, LatentInfo = "LatentInfo", WorldContext = "WorldContextObject", Keywords = "googlear arcore availability", DeprecatedFunction, DeprecationMessage="This function is deprecated, use InstallARService from UARDependencyHandler."))
-	static void InstallARCoreService(UObject* WorldContextObject, struct FLatentActionInfo LatentInfo, EGoogleARCoreInstallRequestResult& OutInstallResult);
-
-	/**
-	 * A polling function to check the ARCore availability in C++.
-	 * This may initiate a query with a remote service to determine if the device is supported by ARCore, so this function will EGoogleARCoreAvailability::UnknownChecking.
-	 *
-	 * @return	The availability result as a EGoogleARCoreAvailability.
-	 */
-	UE_DEPRECATED(4.26, "This function is deprecated, use CheckARServiceAvailability from UARDependencyHandler.")
-	static EGoogleARCoreAvailability CheckARCoreAvailableStatus();
-
-	/**
-	 * Initiates installation of ARCore if required.
-	 * This function will return immediately and may pause your application if installing ARCore is required.
-	 *
-	 * @return EGoogleARCoreInstallStatus::Requrested if it started a install request.
-	 */
-	UE_DEPRECATED(4.26, "This function is deprecated, use InstallARService from UARDependencyHandler.")
-	static EGoogleARCoreInstallStatus RequestInstallARCoreAPK();
-
-	/**
-	 * A polling function to check the ARCore install request result in C++.
-	 * After you call RequestInstallARCoreAPK() and it returns EGoogleARCoreInstallStatus::Requrested. You can call this function to check the install requst result.
-	 *
-	 * @return The install request result as a EGoogleARCoreInstallRequestResult.
-	 */
-	UE_DEPRECATED(4.26, "This function is deprecated, use CheckARServiceAvailability from UARDependencyHandler.")
-	static EGoogleARCoreInstallRequestResult GetARCoreAPKInstallResult();
-
-	/**
-	 * Get the UGoogleARCoreEventManager to bind BP events or c++ delegate in GoogleARCore plugins.
-	 */
-	UE_DEPRECATED(4.26, "This function is deprecated, use GetSupportedVideoFormats.")
-	UFUNCTION(BlueprintPure, Category = "GoogleARCore|Session", meta = (Keywords = "googlear arcore event manager", DeprecatedFunction, DeprecationMessage="This function is deprecated, use GetSupportedVideoFormats."))
-	static UGoogleARCoreEventManager* GetARCoreEventManager();
-
-	/**
-	 * Starts a new ARCore tracking session GoogleARCore specific configuration.
-	 * If the session already started and the config isn't the same, it will stop the previous session and start a new session with the new config.
-	 * Note that this is a latent action, you can query the session start result by querying GetARCoreSessionStatus() after the latent action finished.
-	 *
-	 * @param WorldContextObject	The world context.
-	 * @param LatentInfo			Unreal internal type required for all latent actions.
-	 * @param Configuration			The ARCoreSession configuration to start the session.
-	 */
-	UE_DEPRECATED(4.26, "This function is deprecated, use StartARSession.")
-	UFUNCTION(BlueprintCallable, Category = "GoogleARCore|Session", meta = (Latent, LatentInfo = "LatentInfo", WorldContext = "WorldContextObject", Keywords = "googlear arcore session start config", DeprecatedFunction, DeprecationMessage="This function is deprecated, use StartARSession."))
-	static void StartARCoreSession(UObject* WorldContextObject, struct FLatentActionInfo LatentInfo, UGoogleARCoreSessionConfig* Configuration);
-
-	/**
-	 * Configure the ARCoreSession with the desired camera configuration. The TargetCameraConfig must be
-	 * from a list returned by UGoogleARCoreEventManager::OnCameraConfig delegate.
-	 *
-	 * This function should be called when UGoogleARCoreEventManager::OnCameraConfig delegate got triggered.
-	 */
-	UE_DEPRECATED(4.26, "This function is deprecated, use UARSessionConfig::SetDesiredVideoFormat.")
-	UFUNCTION(BlueprintCallable, Category = "GoogleARCore|Session", meta = (Keywords = "googlear arcore camera config", DeprecatedFunction, DeprecationMessage="This function is deprecated, use UARSessionConfig::SetDesiredVideoFormat."))
-	static bool SetARCoreCameraConfig(FGoogleARCoreCameraConfig TargetCameraConfig);
-
-	/**
-	 * Get the FGoogleARCoreCameraConfig that the current ARCore session is using.
-	 *
-	 * @param OutCurrentCameraConfig   The FGoogleARCoreCameraConfig that the current ARCore session is using.
-	 * @return  True if there is a valid ARCore session and the current camera config is returned.
-	 *          False if ARCore session hasn't been started or it is already stopped.
-	 */
-	UE_DEPRECATED(4.26, "This function is deprecated, use GetSupportedVideoFormats.")
-	UFUNCTION(BlueprintPure, Category = "GoogleARCore|Session", meta = (Keywords = "googlear arcore camera config", DeprecatedFunction, DeprecationMessage="This function is deprecated, use GetSupportedVideoFormats."))
-	static bool GetARCoreCameraConfig(FGoogleARCoreCameraConfig& OutCurrentCameraConfig);
-
-	//-----------------PassthroughCamera---------------------
-	/**
-	 * Returns the state of the passthrough camera rendering in GoogleARCore ARSystem.
-	 *
-	 * @return	True if the passthrough camera rendering is enabled.
-	 */
-	UE_DEPRECATED(4.26, "This function is deprecated, use UARSessionConfig::ShouldRenderCameraOverlay.")
-	UFUNCTION(BlueprintPure, Category = "GoogleARCore|PassthroughCamera", meta = (Keywords = "googlear arcore passthrough camera", DeprecatedFunction, DeprecationMessage="This function is deprecated, use UARSessionConfig::ShouldRenderCameraOverlay."))
-	static bool IsPassthroughCameraRenderingEnabled();
-
-	/**
-	 * Enables/Disables the passthrough camera rendering in GoogleARCore ARSystem.
-	 * Note that when passthrough camera rendering is enabled, the camera FOV will be forced
-	 * to match FOV of the physical camera on the device.
-	 */
-	UE_DEPRECATED(4.26, "This function is deprecated, use UARSessionConfig::bEnableAutomaticCameraOverlay.")
-	UFUNCTION(BlueprintCallable, Category = "GoogleARCore|PassthroughCamera", meta = (Keywords = "googlear arcore passthrough camera", DeprecatedFunction, DeprecationMessage="This function is deprecated, use UARSessionConfig::bEnableAutomaticCameraOverlay."))
-	static void SetPassthroughCameraRenderingEnabled(bool bEnable);
-
-	/**
-	 * Gets the texture coordinate information about the passthrough camera texture.
-	 *
-	 * @param InUV		The original UVs of on the quad. Should be an array with 8 floats.
-	 * @param OutUV		The orientated UVs that can be used to sample the passthrough camera texture and make sure it is displayed correctly.
-	 */
-	UE_DEPRECATED(4.21, "Use UGoogleARCoreFrameFunctionLibrary::TransformARCoordinates2D(EGoogleARCoreCoordinates2DType::Viewport, InUV, EGoogleARCoreCoordinates2DType::Texture, OutUV) instead.")
-	UFUNCTION(BlueprintPure, Category = "GoogleARCore|PassthroughCamera", meta = (Keywords = "googlear arcore passthrough camera uv"))
-	static void GetPassthroughCameraImageUV(const TArray<float>& InUV, TArray<float>& OutUV);
-
-	//-------------------Trackables-------------------------
-	/**
-	 * Gets a list of all valid UARPlaneGeometry objects that ARCore is currently tracking.
-	 * Planes that have entered the EARTrackingState::StoppedTracking state or for which
-	 * UARPlaneGeometry::GetSubsumedBy returns non-null will not be included.
-	 *
-	 * @param OutPlaneList		An array that contains all the valid planes detected by ARCore.
-	 */
-	UE_DEPRECATED(4.26, "This function is deprecated. Use \"GetAllGeometriesByClass\" instead.")
-	UFUNCTION(BlueprintCallable, Category = "GoogleARCore|TrackablePlane", meta = (DeprecatedFunction, DeprecationMessage = "GetAllPlanes is deprecated. Use GetAllGeometriesByClass.", Keywords = "googlear arcore all plane"))
-	static void GetAllPlanes(TArray<UARPlaneGeometry*>& OutPlaneList);
-
-	/**
-	 * Gets a list of all valid UARTrackedPoint objects that ARCore is currently tracking.
-	 * TrackablePoint that have entered the EARTrackingState::StoppedTracking state will not be included.
-	 *
-	 * @param OutTrackablePointList		An array that contains all the valid trackable points detected by ARCore.
-	 */
-	UE_DEPRECATED(4.26, "This function is deprecated. Use \"GetAllGeometriesByClass\" instead.")
-	UFUNCTION(BlueprintCallable, Category = "GoogleARCore|TrackablePoint", meta = (DeprecatedFunction, DeprecationMessage = "GetAllTrackablePoints is deprecated. Use GetAllGeometriesByClass.", Keywords = "googlear arcore pose transform"))
-	static void GetAllTrackablePoints(TArray<UARTrackedPoint*>& OutTrackablePointList);
-
-	/**
-	 * Gets a list of all valid UGoogleARCoreAugmentedImage objects that ARCore is currently tracking.
-	 * UGoogleARCoreAugmentedImage that have entered the EARTrackingState::StoppedTracking state will not be included.
-	 *
-	 * @param OutAugmentedImageList		An array that contains all the valid augmented images detected by ARCore.
-	 */
-	UE_DEPRECATED(4.26, "This function is deprecated. Use \"GetAllGeometriesByClass\" instead.")
-	UFUNCTION(BlueprintCallable, Category = "GoogleARCore|AugmentedImage", meta = (DeprecatedFunction, DeprecationMessage = "GetAllAugmentedImages is deprecated. Use GetAllGeometriesByClass.", Keywords = "googlear arcore all augmented image"))
-	static void GetAllAugmentedImages(TArray<UGoogleARCoreAugmentedImage*>& OutAugmentedImageList);
-
-	/**
-	 * Gets a list of all valid UGoogleARCoreAugmentedFace objects that ARCore is currently tracking.
-	 * UGoogleARCoreAugmentedFace that have entered the EARTrackingState::StoppedTracking state will not be included.
-	 *
-	 * @param OutAugmentedFaceList		An array that contains all the valid augmented faces detected by ARCore.
-	 */
-	UE_DEPRECATED(4.26, "This function is deprecated, use GetAllGeometriesByClass.")
-	UFUNCTION(BlueprintCallable, Category = "GoogleARCore|AugmentedFace", meta = (Keywords = "googlear arcore all augmented face", DeprecatedFunction, DeprecationMessage = "This function is deprecated, use GetAllGeometriesByClass."))
-	static void GetAllAugmentedFaces(TArray<UGoogleARCoreAugmentedFace*>& OutAugmentedFaceList);
-
-	/** Template function to get all trackables from a given type. */
-	template< class T > static void GetAllTrackable(TArray<T*>& OutTrackableList);
-
 	/**
 	 * Create an ARCandidateImage object from the raw pixel data and add it to the ARCandidateImageList of the given \c UARSessionConfig object.
 	 *
@@ -211,149 +39,6 @@ class GOOGLEARCOREBASE_API UGoogleARCoreFrameFunctionLibrary : public UBlueprint
 	GENERATED_BODY()
 
 public:
-	/**
-	 * Returns the current ARCore session status.
-	 *
-	 * @return	A EARSessionStatus enum that describes the session status.
-	 */
-	UE_DEPRECATED(4.26, "This function is deprecated. Use \"GetTrackingQuality\" instead.")
-	UFUNCTION(BlueprintPure, Category = "GoogleARCore|MotionTracking", meta = (DeprecatedFunction, DeprecationMessage = "GetTrackingState is deprecated. Use GetTrackingQuality.", Keywords = "googlear arcore session"))
-	static EGoogleARCoreTrackingState GetTrackingState();
-
-	/**
-	 * Returns the reason when UARBlueprintLibrary::GetTrackingQuality() returns NotTracking, or UGoogleARCoreFrameFunctionLibrary::GetTrackingState
-	 * returns Paused.
-	 *
-	 * In scenarios when multiple causes result in tracking failures, this reports the most actionable failure reason.
-	 *
-	 * @return	A EGoogleARCoreTrackingFailureReason enum that describes the tracking failure reason.
-	 */
-	UE_DEPRECATED(4.26, "This function is deprecated. Use \"GetTrackingQualityReason\" instead.")
-	UFUNCTION(BlueprintPure, Category = "GoogleARCore|MotionTracking", meta = (DeprecatedFunction, DeprecationMessage = "GetTrackingFailureReason is deprecated. Use GetTrackingQualityReason.", Keywords = "googlear arcore session"))
-	static EGoogleARCoreTrackingFailureReason GetTrackingFailureReason();
-
-	/**
-	 * Gets the latest tracking pose of the ARCore device in Unreal AR Tracking Space
-	 *
-	 * Note that ARCore motion tracking has already integrated with HMD and the motion controller interface.
-	 * Use this function only if you need to implement your own tracking component.
-	 *
-	 * @param OutPose		The latest device pose.
-	 * @return				True if the pose is updated successfully for this frame.
-	 */
-	UE_DEPRECATED(4.26, "This function is deprecated, use GetOrientationAndPosition.")
-	UFUNCTION(BlueprintPure, Category = "GoogleARCore|MotionTracking", meta = (Keywords = "googlear arcore pose transform", DeprecatedFunction, DeprecationMessage = "This function is deprecated, use GetOrientationAndPosition."))
-	static void GetPose(FTransform& OutPose);
-
-	/**
-	 * Traces a ray from the user's device in the direction of the given location in the camera
-	 * view. Intersections with detected scene geometry are returned, sorted by distance from the
-	 * device; the nearest intersection is returned first.
-	 *
-	 * @param WorldContextObject	The world context.
-	 * @param ScreenPosition		The position on the screen to cast the ray from.
-	 * @param TraceChannels			A set of EGoogleARCoreLineTraceChannel indicate which type of line trace it should perform.
-	 * @param OutHitResults			The list of hit results sorted by distance.
-	 * @return						True if there is a hit detected.
-	 */
-	UE_DEPRECATED(4.26, "This function is deprecated, use LineTraceTrackedObjects.")
-	UFUNCTION(BlueprintCallable, Category = "GoogleARCore|LineTrace", meta = (WorldContext = "WorldContextObject", Keywords = "googlear arcore raycast hit", DeprecatedFunction, DeprecationMessage = "This function is deprecated, use LineTraceTrackedObjects."))
-	static bool ARCoreLineTrace(UObject* WorldContextObject, const FVector2D& ScreenPosition, TSet<EGoogleARCoreLineTraceChannel> TraceChannels, TArray<FARTraceResult>& OutHitResults);
-
-	/**
-	* Traces a ray along the given line. Intersections with detected scene geometry are returned,
-	* sorted by distance from the start of the line; the nearest intersection is returned first.
-	*
-	* @param WorldContextObject	The world context.
-	* @param Start		The start of line segment.
-	* @param End		The end of line segment.
-	* @param TraceChannels			A set of EGoogleARCoreLineTraceChannel indicate which type of line trace it should perform.
-	* @param OutHitResults			The list of hit results sorted by distance.
-	* @return						True if there is a hit detected.
-	*/
-	UE_DEPRECATED(4.26, "This function is deprecated, use LineTraceTrackedObjects3D.")
-	UFUNCTION(BlueprintCallable, Category = "GoogleARCore|LineTrace", meta = (WorldContext = "WorldContextObject", Keywords = "googlear arcore raycast hit", DeprecatedFunction, DeprecationMessage = "This function is deprecated, use LineTraceTrackedObjects3D."))
-	static bool ARCoreLineTraceRay(UObject* WorldContextObject, const FVector& Start, const FVector& End, TSet<EGoogleARCoreLineTraceChannel> TraceChannels, TArray<FARTraceResult>& OutHitResults);
-
-	/**
-	 * Gets a list of UARPin objects that were changed in this frame.
-	 *
-	 * @param OutAnchorList		An array that contains the updated UARPin.
-	 */
-	UE_DEPRECATED(4.26, "This function is deprecated, use GetAllPins.")
-	UFUNCTION(BlueprintCallable, Category = "GoogleARCore|ARAnchor", meta = (Keywords = "googlear arcore Anchor", DeprecatedFunction, DeprecationMessage="This function is deprecated, use GetAllPins."))
-	static void GetUpdatedARPins(TArray<UARPin*>& OutAnchorList);
-
-	/**
-	 * Gets a list of UARPlaneGeometry objects that were changed in this frame.
-	 *
-	 * @param OutPlaneList	An array that contains the updated UARPlaneGeometry.
-	 */
-	UE_DEPRECATED(4.26, "This function is deprecated, use GetAllGeometriesByClass.")
-	UFUNCTION(BlueprintCallable, Category = "GoogleARCore|TrackablePlane", meta = (Keywords = "googlear arcore pose transform", DeprecatedFunction, DeprecationMessage="This function is deprecated, use GetAllGeometriesByClass."))
-	static void GetUpdatedPlanes(TArray<UARPlaneGeometry*>& OutPlaneList);
-
-	/**
-	 * Gets a list of UARTrackedPoint objects that were changed in this frame.
-	 *
-	 * @param OutTrackablePointList	An array that contains the updated UARTrackedPoint.
-	 */
-	UE_DEPRECATED(4.26, "This function is deprecated, use GetAllGeometriesByClass.")
-	UFUNCTION(BlueprintCallable, Category = "GoogleARCore|TrackablePoint", meta = (Keywords = "googlear arcore pose transform", DeprecatedFunction, DeprecationMessage="This function is deprecated, use GetAllGeometriesByClass."))
-	static void GetUpdatedTrackablePoints(TArray<UARTrackedPoint*>& OutTrackablePointList);
-
-	/**
-	 * Gets a list of UGoogleARCoreAugmentedImage objects that were changed in this frame.
-	 *
-	 * @param OutImageList	An array that contains the updated UGoogleARCoreAugmentedImage.
-	 */
-	UE_DEPRECATED(4.26, "This function is deprecated, use GetAllGeometriesByClass.")
-	UFUNCTION(BlueprintCallable, Category = "GoogleARCore|AugmentedImage", meta = (Keywords = "googlear arcore augmented image updated", DeprecatedFunction, DeprecationMessage="This function is deprecated, use GetAllGeometriesByClass."))
-	static void GetUpdatedAugmentedImages(TArray<UGoogleARCoreAugmentedImage*>& OutImageList);
-
-	/**
-	* Gets a list of UGoogleARCoreAugmentedFace objects that were changed in this frame.
-	*
-	* @param OutFaceList	An array that contains the updated UGoogleARCoreAugmentedFace.
-	*/
-	UE_DEPRECATED(4.26, "This function is deprecated, use GetAllGeometriesByClass.")
-	UFUNCTION(BlueprintCallable, Category = "GoogleARCore|AugmentedFace", meta = (Keywords = "googlear arcore augmented face updated", DeprecatedFunction, DeprecationMessage="This function is deprecated, use GetAllGeometriesByClass."))
-	static void GetUpdatedAugmentedFaces(TArray<UGoogleARCoreAugmentedFace*>& OutFaceList);
-
-	/** Template function to get the updated trackables in this frame a given trackable type. */
-	template< class T > static void GetUpdatedTrackable(TArray<T*>& OutTrackableList);
-
-	/**
-	 * Gets the latest light estimation.
-	 *
-	 * @param OutLightEstimate		The struct that describes the latest light estimation.
-	 */
-	UE_DEPRECATED(4.26, "This function is deprecated, use GetCurrentLightEstimate.")
-	UFUNCTION(BlueprintCallable, Category = "GoogleARCore|LightEstimation", meta = (Keywords = "googlear arcore light ambient", DeprecatedFunction, DeprecationMessage="This function is deprecated, use GetCurrentLightEstimate."))
-	static void GetLightEstimation(FGoogleARCoreLightEstimate& OutLightEstimate);
-
-	/**
-	 * Gets the latest point cloud that will be only available for this frame.
-	 * If you want to keep the point cloud data, you can either copy it to your own struct
-	 * or call AcquireLatestPointCloud() to avoid the copy.
-	 *
-	 * @param OutLatestPointCloud		A pointer point to the latest point cloud.
-	 * @return  An EGoogleARCoreFunctionStatus. Possible value: Success, SessionPaused, ResourceExhausted.
-	 */
-	UE_DEPRECATED(4.26, "This function is deprecated, use GetPointCloud.")
-	UFUNCTION(BlueprintCallable, Category = "GoogleARCore|PointCloud", meta = (Keywords = "googlear arcore pointcloud", DeprecatedFunction, DeprecationMessage="This function is deprecated, use GetPointCloud."))
-	static EGoogleARCoreFunctionStatus GetPointCloud(UGoogleARCorePointCloud*& OutLatestPointCloud);
-
-	/**
-	 * Acquires latest point cloud. This will make the point cloud remain valid unless you call UGoogleARCrePointCloud::ReleasePointCloud().
-	 * Be aware that this function could fail if the maximal number of point cloud has been acquired.
-	 *
-	 * @param OutLatestPointCloud		A pointer point to the latest point cloud.
-	 * @return  An EGoogleARCoreFunctionStatus. Possible value: Success, SessionPaused, ResourceExhausted.
-	 */
-	UE_DEPRECATED(4.26, "This function is deprecated, use GetPointCloud.")
-	UFUNCTION(BlueprintCallable, Category = "GoogleARCore|PointCloud", meta = (Keywords = "googlear arcore pointcloud", DeprecatedFunction, DeprecationMessage="This function is deprecated, use GetPointCloud."))
-	static EGoogleARCoreFunctionStatus AcquirePointCloud(UGoogleARCorePointCloud*& OutLatestPointCloud);
 
 #if PLATFORM_ANDROID
 	/**
@@ -363,63 +48,7 @@ public:
 	 * @param OutCameraMetadata		A pointer to a ACameraMetadata struct which is only valid in one frame.
 	 * @return An EGoogleARCoreFunctionStatus. Possible value: Success, SessionPaused, NotAvailable.
 	 */
+	UE_DEPRECATED(5.3, "GetCameraMetadata is deprecated.  It does not have an exact replacement.  ArImageMetadata seems to be the closest thing, but we do not currently expose it.")
 	static EGoogleARCoreFunctionStatus GetCameraMetadata(const ACameraMetadata*& OutCameraMetadata);
 #endif
-
-	/**
-	 * Get the pass-through camera texture that GoogleARCore plugin will use to render the passthrough camera background.
-	 * Note that UTexture object this function returns may change every frame. If you want to use the camera texture, you should
-	 * call the function every frame and update the texture parameter in your material.
-	 *
-	 * @return A pointer to the UTexture that will be used to render the passthrough camera background.
-	 */
-	UE_DEPRECATED(4.26, "This function is deprecated. Use \"GoogleARCore Passthrough Camera\" expression in your material instead.")
-	UFUNCTION(BlueprintPure, Category = "GoogleARCore|PassthroughCamera", meta = (DeprecatedFunction, DeprecationMessage = "GetCameraTexture is deprecated. Currently, use GoogleARCore Passthrough Camera expression.", Keywords = "googlear arcore passthroughcamera"))
-	static UTexture* GetCameraTexture();
-
-	/**
-	 * Acquire a CPU-accessible camera image.
-	 *
-	 * @param OutLatestCameraImage    A place to store the pointer to a new UGoogleARCoreCameraImage instance.
-	 * @return An EGoogleARCoreFunctionStatus. Possible value: Success, ResourceExhausted, NotAvailable.
-	 */
-	UE_DEPRECATED(4.26, "This function is deprecated, use GetARTexture.")
-	UFUNCTION(BlueprintCallable, Category = "GoogleARCore|PassthroughCamera", meta = (Keywords = "googlear arcore passthroughcamera", DeprecatedFunction, DeprecationMessage = "This function is deprecated, use GetARTexture."))
-	static EGoogleARCoreFunctionStatus AcquireCameraImage(UGoogleARCoreCameraImage *&OutLatestCameraImage);
-
-	/**
-	 * Transforms an array of 2D coordinates into a different 2D coordinate system.  This will account for the display rotation,
-	 * and any additional required adjustment.
-	 *
-	 * Some examples of useful conversions:
-	 *   To transform screen space UVs for texture space UVs to rendering pass-through camera texture: Viewport -> Texture;
-	 *   To transform a point found by a computer vision algorithm in the pass-through camera image into a point on the viewport: Image -> Viewport;
-	 *
-	 * @param InputCoordinatesType		The coordinate system used by InputCoordinates.
-	 * @param InputCoordinates			The input 2d coordinates.
-	 * @param OutputCoordinatesType		The coordinate system to transform to.
-	 * @param OutputCoordinates			The output 2d coordinates.
-	 */
-	UE_DEPRECATED(4.26, "This function is deprecated, use the standard materials from ARUtilities plugin for passthrough rendering.")
-	UFUNCTION(BlueprintPure, Category = "GoogleARCore|PassthroughCamera", meta = (Keywords = "googlear arcore passthrough camera uv", DeprecatedFunction, DeprecationMessage = "This function is deprecated, use the standard materials from ARUtilities plugin for passthrough rendering."))
-	static void TransformARCoordinates2D(EGoogleARCoreCoordinates2DType InputCoordinatesType, const TArray<FVector2D>& InputCoordinates,
-		EGoogleARCoreCoordinates2DType OutputCoordinatesType, TArray<FVector2D>& OutputCoordinates);
-
-	/**
-	 * Get the camera intrinsics for the camera image (CPU image).
-	 *
-	 * @param OutCameraIntrinsics  The output intrinsics object.
-	 */
-	UE_DEPRECATED(4.26, "This function is deprecated, use GetCameraIntrinsics.")
-	UFUNCTION(BlueprintCallable, Category = "GoogleARCore|CameraIntrinsics", meta = (Keywords = "googlear arcore camera", DeprecatedFunction, DeprecationMessage = "This function is deprecated, use GetCameraIntrinsics."))
-	static EGoogleARCoreFunctionStatus GetCameraImageIntrinsics(UGoogleARCoreCameraIntrinsics *&OutCameraIntrinsics);
-
-	/**
-	 * Get the camera intrinsics for the camera texture (GPU image).
-	 *
-	 * @param OutCameraIntrinsics  The output intrinsics object.
-	 */
-	UE_DEPRECATED(4.26, "This function is deprecated, use GetCameraIntrinsics.")
-	UFUNCTION(BlueprintCallable, Category = "GoogleARCore|CameraIntrinsics", meta = (Keywords = "googlear arcore camera", DeprecatedFunction, DeprecationMessage = "This function is deprecated, use GetCameraIntrinsics."))
-	static EGoogleARCoreFunctionStatus GetCameraTextureIntrinsics(UGoogleARCoreCameraIntrinsics *&OutCameraIntrinsics);
 };
diff --git a/Engine/Plugins/Runtime/AR/Google/GoogleARCore/Source/GoogleARCoreBase/Public/GoogleARCoreTypes.h b/Engine/Plugins/Runtime/AR/Google/GoogleARCore/Source/GoogleARCoreBase/Public/GoogleARCoreTypes.h
--- a/Engine/Plugins/Runtime/AR/Google/GoogleARCore/Source/GoogleARCoreBase/Public/GoogleARCoreTypes.h
+++ b/Engine/Plugins/Runtime/AR/Google/GoogleARCore/Source/GoogleARCoreBase/Public/GoogleARCoreTypes.h
@@ -26,7 +26,6 @@ typedef struct ArPointCloud_ ArPointCloud;
 typedef struct ArAnchor_ ArAnchor;
 #endif
 
-UENUM(BlueprintType)
 enum class EGoogleARCoreAvailability : uint8
 {
 	/* An internal error occurred while determining ARCore availability. */
@@ -51,7 +50,6 @@ enum class EGoogleARCoreAvailability : uint8
 	SupportedInstalled = 203
 };
 
-UENUM(BlueprintType)
 enum class EGoogleARCoreInstallStatus : uint8
 {
 	/* The requested resource is already installed.*/
@@ -60,19 +58,6 @@ enum class EGoogleARCoreInstallStatus : uint8
 	Requrested = 1,
 };
 
-UENUM(BlueprintType)
-enum class EGoogleARCoreInstallRequestResult : uint8
-{
-	/* The ARCore APK is installed*/
-	Installed,
-	/* ARCore APK install request failed because the device is not compatible. */
-	DeviceNotCompatible,
-	/* ARCore APK install request failed because the current version of android is too old to support ARCore. */
-	UserDeclinedInstallation,
-	/* ARCore APK install request failed because unknown error happens while checking or requesting installation. */
-	FatalError
-};
-
 /**
  * @ingroup GoogleARCoreBase
  * Describes the status of most ARCore functions.
@@ -104,7 +89,6 @@ enum class EGoogleARCoreFunctionStatus : uint8
  * @ingroup GoogleARCoreBase
  * Describes the tracking state of the current ARCore session.
  */
-UENUM(BlueprintType)
 enum class EGoogleARCoreTrackingState : uint8
 {
 	/** Tracking is valid. */
@@ -119,7 +103,6 @@ enum class EGoogleARCoreTrackingState : uint8
 * @ingroup GoogleARCoreBase
 * 2d coordinate systems referenced by ARCore.
 */
-UENUM(BlueprintType)
 enum class EGoogleARCoreCoordinates2DType : uint8
 {
 	/** ARCore normalized pass-through camera texture(for GPU usage) space with top-left (0.0f, 0.0f) and bottom-right(1.0f, 1.0f). */
@@ -134,23 +117,17 @@ enum class EGoogleARCoreCoordinates2DType : uint8
 /**
  * A struct describes the ARCore light estimation.
  */
-USTRUCT(BlueprintType)
 struct FGoogleARCoreLightEstimate
 {
-	GENERATED_BODY()
-
 	/** Whether this light estimation is valid. */
-	UPROPERTY(BlueprintReadOnly, Category = "GoogleARCore|LightEstimate")
 	bool bIsValid = false;
 
 	/** The average pixel intensity of the passthrough camera image. */
-	UPROPERTY(BlueprintReadOnly, Category = "GoogleARCore|LightEstimate")
 	float PixelIntensity = 0.0f;
 
 	/**
 	 * The RGB scale to match the color of the light in the real environment.
 	 */
-	UPROPERTY(BlueprintReadOnly, Category = "GoogleARCore|LightEstimate")
 	FVector RGBScaleFactor = FVector(0.0f);
 };
 
