UE-189163 Unknown ARCore Trackable Type assert during ARLineTrace
-Ignoring and releasing any DEPTH_POINT trackables that we find so that they do not assert.
-A couple small logging improvements.
-This fixes the HandheldAR template project.

diff --git a/Engine/Plugins/Runtime/AR/Google/GoogleARCore/Source/GoogleARCoreBase/Private/GoogleARCoreAPI.cpp b/Engine/Plugins/Runtime/AR/Google/GoogleARCore/Source/GoogleARCoreBase/Private/GoogleARCoreAPI.cpp
--- a/Engine/Plugins/Runtime/AR/Google/GoogleARCore/Source/GoogleARCoreBase/Private/GoogleARCoreAPI.cpp
+++ b/Engine/Plugins/Runtime/AR/Google/GoogleARCore/Source/GoogleARCoreBase/Private/GoogleARCoreAPI.cpp
@@ -1032,20 +1032,26 @@ void FGoogleARCoreFrame::Update(float WorldToMeterScale)
 		for (auto TrackableHandle : Trackables)
 		{
 			// Note that this will create a new UARTrackedGeometry if there's no existing record for it for the handle
-			const auto& Group = ObjectManager->GetBaseTrackableFromHandle(TrackableHandle, Session);
-			// Updated the cached tracked geometry when it is valid.
-			FGoogleARCoreTrackableResource* TrackableResource = reinterpret_cast<FGoogleARCoreTrackableResource*>(Group.TrackedGeometry->GetNativeResource());
-			TrackableResource->UpdateGeometryData();
-			
-			//trigger delegates
-			if (Group.ARComponent)
+			// Note also that there are trackable types we do not support and for which it will return nullptr.
+			const FTrackedGeometryGroup* GroupPtr = ObjectManager->GetBaseTrackableFromHandle(TrackableHandle, Session);
+			if (GroupPtr)
 			{
-				Group.ARComponent->Update(Group.TrackedGeometry);
-			}
+				const FTrackedGeometryGroup& Group = *GroupPtr;
+
+				// Updated the cached tracked geometry when it is valid.
+				FGoogleARCoreTrackableResource* TrackableResource = reinterpret_cast<FGoogleARCoreTrackableResource*>(Group.TrackedGeometry->GetNativeResource());
+				TrackableResource->UpdateGeometryData();
 			
-			if (TrackingSystem)
-			{
-				TrackingSystem->OnTrackableUpdated(Group.TrackedGeometry);
+				//trigger delegates
+				if (Group.ARComponent)
+				{
+					Group.ARComponent->Update(Group.TrackedGeometry);
+				}
+			
+				if (TrackingSystem)
+				{
+					TrackingSystem->OnTrackableUpdated(Group.TrackedGeometry);
+				}
 			}
 		}
 		
@@ -1462,6 +1468,7 @@ void FGoogleARCoreFrame::FilterLineTraceResults(ArHitResultList* HitResultList,
 			if (OrientationMode == AR_POINT_ORIENTATION_ESTIMATED_SURFACE_NORMAL && !!(RequestedTraceChannels & EGoogleARCoreLineTraceChannel::FeaturePointWithSurfaceNormal))
 			{
 				UARTrackedGeometry* TrackedGeometry = Session->GetUObjectManager()->GetTrackableFromHandle<UARTrackedGeometry>(TrackableHandle, Session);
+				check(TrackedGeometry);
 				FARTraceResult UEHitResult(Session->GetARSystem(), Distance, EARLineTraceChannels::FeaturePoint, HitTransform, TrackedGeometry);
 				OutHitResults.Add(UEHitResult);
 				continue;
@@ -1469,6 +1476,7 @@ void FGoogleARCoreFrame::FilterLineTraceResults(ArHitResultList* HitResultList,
 			if (!!(RequestedTraceChannels & EGoogleARCoreLineTraceChannel::FeaturePoint))
 			{
 				UARTrackedGeometry* TrackedGeometry = Session->GetUObjectManager()->GetTrackableFromHandle<UARTrackedGeometry>(TrackableHandle, Session);
+				check(TrackedGeometry);
 				FARTraceResult UEHitResult(Session->GetARSystem(), Distance, EARLineTraceChannels::FeaturePoint, HitTransform, TrackedGeometry);
 				OutHitResults.Add(UEHitResult);
 				continue;
@@ -1484,6 +1492,7 @@ void FGoogleARCoreFrame::FilterLineTraceResults(ArHitResultList* HitResultList,
 				if (PointInsidePolygon)
 				{
 					UARTrackedGeometry* TrackedGeometry = Session->GetUObjectManager()->GetTrackableFromHandle<UARTrackedGeometry>(TrackableHandle, Session);
+					check(TrackedGeometry);
 					FARTraceResult UEHitResult(Session->GetARSystem(), Distance, EARLineTraceChannels::PlaneUsingBoundaryPolygon, HitTransform, TrackedGeometry);
 					OutHitResults.Add(UEHitResult);
 					continue;
@@ -1496,6 +1505,7 @@ void FGoogleARCoreFrame::FilterLineTraceResults(ArHitResultList* HitResultList,
 				if (PointInsideExtents)
 				{
 					UARTrackedGeometry* TrackedGeometry = Session->GetUObjectManager()->GetTrackableFromHandle<UARTrackedGeometry>(TrackableHandle, Session);
+					check(TrackedGeometry);
 					FARTraceResult UEHitResult(Session->GetARSystem(), Distance, EARLineTraceChannels::PlaneUsingExtent, HitTransform, TrackedGeometry);
 					OutHitResults.Add(UEHitResult);
 					continue;
@@ -1504,6 +1514,7 @@ void FGoogleARCoreFrame::FilterLineTraceResults(ArHitResultList* HitResultList,
 			if (!!(RequestedTraceChannels & EGoogleARCoreLineTraceChannel::InfinitePlane))
 			{
 				UARTrackedGeometry* TrackedGeometry = Session->GetUObjectManager()->GetTrackableFromHandle<UARTrackedGeometry>(TrackableHandle, Session);
+				check(TrackedGeometry);
 				FARTraceResult UEHitResult(Session->GetARSystem(), Distance, EARLineTraceChannels::GroundPlane, HitTransform, TrackedGeometry);
 				OutHitResults.Add(UEHitResult);
 				continue;
@@ -1514,6 +1525,7 @@ void FGoogleARCoreFrame::FilterLineTraceResults(ArHitResultList* HitResultList,
 			if (!!(RequestedTraceChannels & EGoogleARCoreLineTraceChannel::AugmentedImage))
 			{
 				UARTrackedGeometry* TrackedGeometry = Session->GetUObjectManager()->GetTrackableFromHandle<UARTrackedGeometry>(TrackableHandle, Session);
+				check(TrackedGeometry);
 				FARTraceResult UEHitResult(Session->GetARSystem(), Distance, EARLineTraceChannels::PlaneUsingExtent, HitTransform, TrackedGeometry);
 				OutHitResults.Add(UEHitResult);
 				continue;
@@ -1884,7 +1896,7 @@ void UGoogleARCoreUObjectManager::ClearTrackables()
 }
 
 #if PLATFORM_ANDROID
-const FTrackedGeometryGroup& UGoogleARCoreUObjectManager::GetBaseTrackableFromHandle(ArTrackable* TrackableHandle, FGoogleARCoreSession* Session)
+const FTrackedGeometryGroup* UGoogleARCoreUObjectManager::GetBaseTrackableFromHandle(ArTrackable* TrackableHandle, FGoogleARCoreSession* Session)
 {
 	if (!TrackableHandleMap.Contains(TrackableHandle)
 		|| (TrackableHandleMap[TrackableHandle].TrackedGeometry == nullptr)
@@ -1920,6 +1932,18 @@ const FTrackedGeometryGroup& UGoogleARCoreUObjectManager::GetBaseTrackableFromHa
 			NativeResource = new FGoogleARCoreAugmentedFaceResource(Session->AsShared(), TrackableHandle, NewTrackableObject);
 			ARComponentClass = AccessSessionConfig().GetFaceComponentClass();
 		}
+		else if (TrackableType == ArTrackableType::AR_TRACKABLE_DEPTH_POINT)
+		{
+			// If depth sensing is enabled to get the depth texture ar ray casts will create 'depth point' trackables.
+			// However depth point trackables position can only be accessed by creating an ArAnchor from them.  This breaks a lot of assumptions in the 
+			// ARTrackedGeometry system design.  Unless there is some refactoring this one trackable will work quite differently from the others.
+			// For now we will just ignore these depth point trackables.
+			ArTrackable_release(TrackableHandle);
+			return nullptr;
+		}
+		// Currently we do not enable this trackable type with ArConfig_setInstantPlacementMode, so we should never have one.
+		check(TrackableType != ArTrackableType::AR_TRACKABLE_INSTANT_PLACEMENT_POINT);
+		
 		// We should have a valid trackable object now.
 		checkf(NewTrackableObject, TEXT("Unknown ARCore Trackable Type: %d"), TrackableType);
 
@@ -1936,8 +1960,8 @@ const FTrackedGeometryGroup& UGoogleARCoreUObjectManager::GetBaseTrackableFromHa
 		// Add the new object to the record
 		TrackableHandleMap.Add(TrackableHandle, TrackedGeometryGroup);
 		
-		UE_LOG(LogGoogleARCoreAPI, Log, TEXT("Added NewTrackableObject: %s for handle 0x%x, current total number of trackables: %d"),
-			   *NewTrackableObject->GetName(), TrackableHandle, TrackableHandleMap.Num());
+		UE_LOG(LogGoogleARCoreAPI, Log, TEXT("Added NewTrackableObject: %s %s for handle 0x%x of type %i, current total number of trackables: %d"),
+			*NewTrackableObject->GetName(), *NewTrackableObject->GetDebugName().ToString(), TrackableHandle, (int32)TrackableType, TrackableHandleMap.Num());
 		
 		const auto Guid = TrackableHandleToGuid(TrackableHandle);
 		NewTrackableObject->UniqueId = Guid;
@@ -1951,7 +1975,7 @@ const FTrackedGeometryGroup& UGoogleARCoreUObjectManager::GetBaseTrackableFromHa
 	
 	const auto& Group = TrackableHandleMap[TrackableHandle];
 	checkf(Group.TrackedGeometry, TEXT("UGoogleARCoreUObjectManager failed to get a valid trackable %p from the map."), TrackableHandle);
-	return Group;
+	return &Group;
 }
 
 void UGoogleARCoreUObjectManager::RemoveInvalidTrackables(const TArray<ArTrackable*>& ValidTrackables, TArray<UARPin*>& ARPinsToRemove)
diff --git a/Engine/Plugins/Runtime/AR/Google/GoogleARCore/Source/GoogleARCoreBase/Private/GoogleARCoreAPI.h b/Engine/Plugins/Runtime/AR/Google/GoogleARCore/Source/GoogleARCoreBase/Private/GoogleARCoreAPI.h
--- a/Engine/Plugins/Runtime/AR/Google/GoogleARCore/Source/GoogleARCoreBase/Private/GoogleARCoreAPI.h
+++ b/Engine/Plugins/Runtime/AR/Google/GoogleARCore/Source/GoogleARCoreBase/Private/GoogleARCoreAPI.h
@@ -205,13 +205,13 @@ public:
 	UARTrackedGeometry* RemoveTrackable(FARCorePointer Pointer);
 
 #if PLATFORM_ANDROID
-	const FTrackedGeometryGroup& GetBaseTrackableFromHandle(ArTrackable* TrackableHandle, FGoogleARCoreSession* Session);
+	const FTrackedGeometryGroup* GetBaseTrackableFromHandle(ArTrackable* TrackableHandle, FGoogleARCoreSession* Session);
 	
 	template<class T>
 	T* GetTrackableFromHandle(ArTrackable* TrackableHandle, FGoogleARCoreSession* Session)
 	{
-		const auto& Group = GetBaseTrackableFromHandle(TrackableHandle, Session);
-		return CastChecked<T>(Group.TrackedGeometry);
+		const FTrackedGeometryGroup* Group = GetBaseTrackableFromHandle(TrackableHandle, Session);
+		return Group ? CastChecked<T>(Group->TrackedGeometry) : nullptr;
 	}
 	
 	void DumpTrackableHandleMap(const ArSession* SessionHandle);
@@ -411,8 +411,10 @@ void FGoogleARCoreFrame::GetUpdatedTrackables(TArray<T*>& OutARCoreTrackableList
 	for (auto TrackableHandle : Trackables)
 	{
 		T* TrackableObject = Session->GetUObjectManager()->template GetTrackableFromHandle<T>(TrackableHandle, Session);
-
-		OutARCoreTrackableList.Add(TrackableObject);
+		if (TrackableObject)
+		{
+			OutARCoreTrackableList.Add(TrackableObject);
+		}
 	}
 	ArTrackableList_destroy(TrackableListHandle);
 #endif
@@ -443,7 +445,10 @@ void FGoogleARCoreSession::GetAllTrackables(TArray<T*>& OutARCoreTrackableList)
 	for (auto TrackableHandle : Trackables)
 	{
 		T* TrackableObject = UObjectManager->template GetTrackableFromHandle<T>(TrackableHandle, this);
-		OutARCoreTrackableList.Add(TrackableObject);
+		if (TrackableObject)
+		{
+			OutARCoreTrackableList.Add(TrackableObject);
+		}
 	}
 	ArTrackableList_destroy(TrackableListHandle);
 #endif
diff --git a/Engine/Plugins/Runtime/AR/Google/GoogleARCore/Source/GoogleARCoreBase/Private/GoogleARCoreXRTrackingSystem.cpp b/Engine/Plugins/Runtime/AR/Google/GoogleARCore/Source/GoogleARCoreBase/Private/GoogleARCoreXRTrackingSystem.cpp
--- a/Engine/Plugins/Runtime/AR/Google/GoogleARCore/Source/GoogleARCoreBase/Private/GoogleARCoreXRTrackingSystem.cpp
+++ b/Engine/Plugins/Runtime/AR/Google/GoogleARCore/Source/GoogleARCoreBase/Private/GoogleARCoreXRTrackingSystem.cpp
@@ -560,7 +560,7 @@ void FGoogleARCoreXRTrackingSystem::OnTrackableAdded(UARTrackedGeometry* InTrack
 {
 	if (ensure(InTrackedGeometry))
 	{
-		UE_LOG(LogGoogleARCoreTrackingSystem, Log, TEXT("OnTrackableAdded: %s"), *InTrackedGeometry->GetName());
+		UE_LOG(LogGoogleARCoreTrackingSystem, Log, TEXT("OnTrackableAdded: %s %s"), *InTrackedGeometry->GetName(), *InTrackedGeometry->GetDebugName().ToString());
 		TriggerOnTrackableAddedDelegates(InTrackedGeometry);
 	}
 }
@@ -577,7 +577,7 @@ void FGoogleARCoreXRTrackingSystem::OnTrackableRemoved(UARTrackedGeometry* InTra
 {
 	if (ensure(InTrackedGeometry))
 	{
-		UE_LOG(LogGoogleARCoreTrackingSystem, Log, TEXT("OnTrackableRemoved: %s"), *InTrackedGeometry->GetName());
+		UE_LOG(LogGoogleARCoreTrackingSystem, Log, TEXT("OnTrackableRemoved: %s %s"), *InTrackedGeometry->GetName(), *InTrackedGeometry->GetDebugName().ToString());
 		TriggerOnTrackableRemovedDelegates(InTrackedGeometry);
 	}
 }
