UE-189163 Unknown ARCore Trackable Type assert during ARLineTrace
-Spotted a case where an ArTrackable handle is not released, fixing it.

diff --git a/Engine/Plugins/Runtime/AR/Google/GoogleARCore/Source/GoogleARCoreBase/Private/GoogleARCoreAPI.cpp b/Engine/Plugins/Runtime/AR/Google/GoogleARCore/Source/GoogleARCoreBase/Private/GoogleARCoreAPI.cpp
--- a/Engine/Plugins/Runtime/AR/Google/GoogleARCore/Source/GoogleARCoreBase/Private/GoogleARCoreAPI.cpp
+++ b/Engine/Plugins/Runtime/AR/Google/GoogleARCore/Source/GoogleARCoreBase/Private/GoogleARCoreAPI.cpp
@@ -1531,6 +1531,9 @@ void FGoogleARCoreFrame::FilterLineTraceResults(ArHitResultList* HitResultList,
 				continue;
 			}
 		}
+
+		// If we don't pass the handle into GetTrackableFromHandle we need to release it here.
+		ArTrackable_release(TrackableHandle);
 	}
 
 	ArHitResult_destroy(HitResultHandle);
