Fix compile error in GoogleARCorePassthroughCameraRenderer.cpp

diff --git a/Engine/Plugins/Runtime/AR/Google/GoogleARCore/Source/GoogleARCoreRendering/Private/GoogleARCorePassthroughCameraRenderer.cpp b/Engine/Plugins/Runtime/AR/Google/GoogleARCore/Source/GoogleARCoreRendering/Private/GoogleARCorePassthroughCameraRenderer.cpp
--- a/Engine/Plugins/Runtime/AR/Google/GoogleARCore/Source/GoogleARCoreRendering/Private/GoogleARCorePassthroughCameraRenderer.cpp
+++ b/Engine/Plugins/Runtime/AR/Google/GoogleARCore/Source/GoogleARCoreRendering/Private/GoogleARCorePassthroughCameraRenderer.cpp
@@ -189,7 +189,8 @@ void FGoogleARCorePassthroughCameraRenderer::RenderVideoOverlayWithMaterial(FRHI
 		return;
 	}
 	
-	SCOPED_DRAW_EVENTF(RHICmdList, RenderVideoOverlay, bRenderingOcclusion ? TEXT("VideoOverlay (Occlusion)") : TEXT("VideoOverlay (Background)"));
+	SCOPED_CONDITIONAL_DRAW_EVENTF(RHICmdList, RenderVideoOverlay_Occlusion ,  bRenderingOcclusion, TEXT("VideoOverlay (Occlusion)" ));
+	SCOPED_CONDITIONAL_DRAW_EVENTF(RHICmdList, RenderVideoOverlay_Background, !bRenderingOcclusion, TEXT("VideoOverlay (Background)"));
 
 	const auto FeatureLevel = InView.GetFeatureLevel();
 	IRendererModule& RendererModule = GetRendererModule();
