UE-207641 GoogleARCoreXRTrackingSystem.cpp warning V1078: An empty container is iterated. The loop will not be executed.
-Just supressing the static analysis warning.
-This code only runs on android, and the guts of the function that gets the array are #ifdef'd out such that the array is always empty.  In an ideal world the module would probably be broken up so that there would be an editor module and a target module so none of this code would build for win64, etc.

diff --git a/Engine/Plugins/Runtime/AR/Google/GoogleARCore/Source/GoogleARCoreBase/Private/GoogleARCoreXRTrackingSystem.cpp b/Engine/Plugins/Runtime/AR/Google/GoogleARCore/Source/GoogleARCoreBase/Private/GoogleARCoreXRTrackingSystem.cpp
--- a/Engine/Plugins/Runtime/AR/Google/GoogleARCore/Source/GoogleARCoreBase/Private/GoogleARCoreXRTrackingSystem.cpp
+++ b/Engine/Plugins/Runtime/AR/Google/GoogleARCore/Source/GoogleARCoreBase/Private/GoogleARCoreXRTrackingSystem.cpp
@@ -446,7 +446,7 @@ TArray<FARVideoFormat> FGoogleARCoreXRTrackingSystem::OnGetSupportedVideoFormats
 
 		TArray<FARVideoFormat> VideoFormats;
 
-		for (const FGoogleARCoreCameraConfig& CameraConfig : SupportedCameraConfig)
+		for (const FGoogleARCoreCameraConfig& CameraConfig : SupportedCameraConfig) //-V1078  // Array always empty on non-android, but this code only actually runs on android.
 		{
 			// FPS is not exposed so we assume it's 30...
 			// Note that we're using CameraTextureResolution rather than CameraImageResolution as the former is relevant
