Fix for exit crash on EngineTest.BootTest

diff --git a/Engine/Plugins/Runtime/AR/Google/GoogleARCore/Source/GoogleARCoreBase/Private/GoogleARCoreAPI.cpp b/Engine/Plugins/Runtime/AR/Google/GoogleARCore/Source/GoogleARCoreBase/Private/GoogleARCoreAPI.cpp
--- a/Engine/Plugins/Runtime/AR/Google/GoogleARCore/Source/GoogleARCoreBase/Private/GoogleARCoreAPI.cpp
+++ b/Engine/Plugins/Runtime/AR/Google/GoogleARCore/Source/GoogleARCoreBase/Private/GoogleARCoreAPI.cpp
@@ -1644,8 +1644,11 @@ void UGoogleARCorePointCloud::GetPointInTrackingSpace(int Index, FVector& OutTra
 void UGoogleARCorePointCloud::ReleasePointCloud()
 {
 #if PLATFORM_ANDROID
-	ArPointCloud_release(PointCloudHandle);
-	PointCloudHandle = nullptr;
+	if (PointCloudHandle)
+	{
+		ArPointCloud_release(PointCloudHandle);
+		PointCloudHandle = nullptr;
+	}
 #endif
 }
 
