GoogleAr material shaders are only ever used on android (via PLATFORM_ANDROID) but were compiling for all "IsMobile" platforms.

diff --git a/Engine/Plugins/Runtime/AR/Google/GoogleARCore/Source/GoogleARCoreRendering/Private/GoogleARCorePassthroughCameraRenderer.cpp b/Engine/Plugins/Runtime/AR/Google/GoogleARCore/Source/GoogleARCoreRendering/Private/GoogleARCorePassthroughCameraRenderer.cpp
--- a/Engine/Plugins/Runtime/AR/Google/GoogleARCore/Source/GoogleARCoreRendering/Private/GoogleARCorePassthroughCameraRenderer.cpp
+++ b/Engine/Plugins/Runtime/AR/Google/GoogleARCore/Source/GoogleARCoreRendering/Private/GoogleARCorePassthroughCameraRenderer.cpp
@@ -96,7 +96,10 @@ public:
 
 	static bool ShouldCompilePermutation(const FMaterialShaderPermutationParameters& Parameters)
 	{
-		return Parameters.MaterialParameters.MaterialDomain == MD_PostProcess && IsMobilePlatform(Parameters.Platform);
+		return Parameters.MaterialParameters.MaterialDomain == MD_PostProcess
+			&& IsMobilePlatform(Parameters.Platform)
+			// Note: IsMobilePlatform can be set for non-android platforms
+			&& IsAndroidPlatform(Parameters.Platform);
 	}
 
 	static void ModifyCompilationEnvironment(const FMaterialShaderPermutationParameters& Parameters, FShaderCompilerEnvironment& OutEnvironment)
