ARCore Vulkan Support
-PR #12921: Vulkan support for GoogleARCore plugin on Android.
-ARCore with OpenGL is now officially deprecated.  It is also non-functional in UE 5.6 because an opengl context refactor removed the shared contexts available on the game thread. A significant refactor of ARCore support would be needed to get it working in 5.6. 5.4 is the reccomended version with OpenGL support for ARCore working correctly.
#SPL_103833

diff --git a/Engine/Plugins/Runtime/AR/Google/GoogleARCore/Source/GoogleARCoreBase/GoogleARCoreBase.Build.cs b/Engine/Plugins/Runtime/AR/Google/GoogleARCore/Source/GoogleARCoreBase/GoogleARCoreBase.Build.cs
--- a/Engine/Plugins/Runtime/AR/Google/GoogleARCore/Source/GoogleARCoreBase/GoogleARCoreBase.Build.cs
+++ b/Engine/Plugins/Runtime/AR/Google/GoogleARCore/Source/GoogleARCoreBase/GoogleARCoreBase.Build.cs
@@ -54,6 +54,10 @@ namespace UnrealBuildTool.Rules
 				PrivateDependencyModuleNames.AddRange(
 					new string[]
 					{
+						// Vulkan support for ARCore
+						"Vulkan",
+						"VulkanRHI",
+
 						"Launch",
 						"OpenGLDrv"
 					}
diff --git a/Engine/Plugins/Runtime/AR/Google/GoogleARCore/Source/GoogleARCoreBase/Private/GoogleARCoreAPI.cpp b/Engine/Plugins/Runtime/AR/Google/GoogleARCore/Source/GoogleARCoreBase/Private/GoogleARCoreAPI.cpp
--- a/Engine/Plugins/Runtime/AR/Google/GoogleARCore/Source/GoogleARCoreBase/Private/GoogleARCoreAPI.cpp
+++ b/Engine/Plugins/Runtime/AR/Google/GoogleARCore/Source/GoogleARCoreBase/Private/GoogleARCoreAPI.cpp
@@ -24,6 +24,7 @@ DEFINE_LOG_CATEGORY(LogGoogleARCoreAPI);
 DECLARE_CYCLE_STAT(TEXT("Update AR Session"), STAT_UpdateARSession, STATGROUP_ARCore);
 DECLARE_CYCLE_STAT(TEXT("Update AR Frame"), STAT_UpdateARFrame, STATGROUP_ARCore);
 DECLARE_CYCLE_STAT(TEXT("Update AR Frame - Camera"), STAT_UpdateCamera, STATGROUP_ARCore);
+DECLARE_CYCLE_STAT(TEXT("Update AR Frame - CameraImage"), STAT_UpdateCameraImage, STATGROUP_ARCore);
 DECLARE_CYCLE_STAT(TEXT("Update AR Frame - DepthImage"), STAT_UpdateDepthImage, STATGROUP_ARCore);
 DECLARE_CYCLE_STAT(TEXT("Update AR Frame - PointCloud"), STAT_UpdatePointCloud, STATGROUP_ARCore);
 DECLARE_CYCLE_STAT(TEXT("Update AR Frame - Trackables"), STAT_UpdateTrackables, STATGROUP_ARCore);
@@ -438,6 +439,12 @@ EGoogleARCoreAPIStatus FGoogleARCoreSession::ConfigSession(const UARSessionConfi
 	
 	ArConfig_setDepthMode(SessionHandle, ConfigHandle, bShouldEnableSceneDepth ? AR_DEPTH_MODE_AUTOMATIC : AR_DEPTH_MODE_DISABLED);
 
+	// Enabling mode to expose Hardware Buffer in Vulkan
+	if (FAndroidMisc::ShouldUseVulkan())
+	{
+		ArConfig_setTextureUpdateMode(SessionHandle, ConfigHandle, AR_TEXTURE_UPDATE_MODE_EXPOSE_HARDWARE_BUFFER);
+	}
+
 	ConfigStatus = ToARCoreAPIStatus(ArSession_configure(SessionHandle, ConfigHandle));
 #endif
 	return ConfigStatus;
@@ -746,6 +753,9 @@ EGoogleARCoreAPIStatus FGoogleARCoreSession::Update(float WorldToMeterScale)
 		FrameNumber++;
 	}
 
+	// If the timestamp is the same then the frame contains the same data.
+	bHasNewFrame = (LastFrameTimestamp != LatestFrame->GetCameraTimestamp());
+
 	return UpdateStatus;
 }
 
@@ -766,7 +776,11 @@ void FGoogleARCoreSession::SetCameraTextureIds(const TArray<uint32_t>& TextureId
 	{
 		return;
 	}
-	ArSession_setCameraTextureNames(SessionHandle, TextureIds.Num(), TextureIds.GetData());
+	// Camera textures names are not used with Vulkan since it relies on getting Hardware Buffer instead.
+	if (!FAndroidMisc::ShouldUseVulkan())
+	{
+		ArSession_setCameraTextureNames(SessionHandle, TextureIds.Num(), TextureIds.GetData());
+	}
 #endif
 }
 
@@ -949,6 +963,23 @@ void FGoogleARCoreFrame::Init()
 #endif
 }
 
+bool FGoogleARCoreFrame::UpdateCameraTexture(UARCoreCameraTexture* InCameraTexture) const
+{
+#if PLATFORM_ANDROID
+	SCOPE_CYCLE_COUNTER(STAT_UpdateCameraImage);
+
+	if (!VulkanHardwareBuffer || !InCameraTexture)
+	{
+		return false;
+	}
+
+	InCameraTexture->UpdateCameraImage(VulkanHardwareBuffer);
+	return true;
+#else
+	return false;
+#endif // PLATFORM_ANDROID
+}
+
 void FGoogleARCoreFrame::UpdateDepthTexture(UARCoreDepthTexture*& OutDepthTexture) const
 {
 #if PLATFORM_ANDROID
@@ -978,7 +1009,17 @@ void FGoogleARCoreFrame::Update(float WorldToMeterScale)
 	
 	{
 		SCOPE_CYCLE_COUNTER(STAT_UpdateCamera);
-		ArFrame_getCameraTextureName(SessionHandle, FrameHandle, &CameraTextureId);
+
+		// With Vulkan obtain the Hardware Buffer used by this frame.
+		if (FAndroidMisc::ShouldUseVulkan())
+		{
+			ArFrame_getHardwareBuffer(SessionHandle, FrameHandle, &VulkanHardwareBuffer);
+		}
+		// With OpenGL obtain the texture Id used by this frame.
+		else
+		{
+			ArFrame_getCameraTextureName(SessionHandle, FrameHandle, &CameraTextureId);
+		}
 
 		ArCamera_release(CameraHandle);
 		ArFrame_acquireCamera(SessionHandle, FrameHandle, &CameraHandle);
diff --git a/Engine/Plugins/Runtime/AR/Google/GoogleARCore/Source/GoogleARCoreBase/Private/GoogleARCoreAPI.h b/Engine/Plugins/Runtime/AR/Google/GoogleARCore/Source/GoogleARCoreBase/Private/GoogleARCoreAPI.h
--- a/Engine/Plugins/Runtime/AR/Google/GoogleARCore/Source/GoogleARCoreBase/Private/GoogleARCoreAPI.h
+++ b/Engine/Plugins/Runtime/AR/Google/GoogleARCore/Source/GoogleARCoreBase/Private/GoogleARCoreAPI.h
@@ -20,6 +20,7 @@
 #include "GoogleARCoreAPI.generated.h"
 
 DECLARE_LOG_CATEGORY_EXTERN(LogGoogleARCoreAPI, Log, All);
+class UARCoreCameraTexture;
 class UARCoreDepthTexture;
 
 enum class EGoogleARCoreAPIStatus : int
@@ -294,6 +295,8 @@ public:
 	
 	bool IsSceneDepthEnabled() const;
 
+	bool HasNewFrame() const { return bHasNewFrame; }
+
 private:
 	virtual void AddReferencedObjects(FReferenceCollector& Collector) override;
 	virtual FString GetReferencerName() const override
@@ -314,6 +317,9 @@ private:
 	ArConfig* ConfigHandle = nullptr;
 	TMap<const UARSessionConfig*, ArAugmentedImageDatabase*> ImageDatabaseMap;
 #endif
+
+	// Stores if current frame has new data. Set on every Update() call.
+	bool bHasNewFrame = false;
 };
 
 class FGoogleARCoreFrame
@@ -327,7 +333,8 @@ public:
 	void Init();
 
 	void Update(float WorldToMeterScale);
-	
+
+	bool UpdateCameraTexture(UARCoreCameraTexture* InCameraTexture) const;
 	void UpdateDepthTexture(UARCoreDepthTexture*& OutDepthTexture) const;
 
 	FTransform GetCameraPose() const;
@@ -382,6 +389,8 @@ private:
 	ArPose* SketchPoseHandle = nullptr;
 	ArImageMetadata* LatestImageMetadata = nullptr;
 
+	void* VulkanHardwareBuffer = nullptr; // This is used instead of CameraTextureId when using Vulkan
+
 	void FilterLineTraceResults(ArHitResultList* HitResultList, EGoogleARCoreLineTraceChannel RequestedTraceChannels,
 		TArray<FARTraceResult>& OutHitResults, float MaxDistance = TNumericLimits<float>::Max()) const;
 #endif
diff --git a/Engine/Plugins/Runtime/AR/Google/GoogleARCore/Source/GoogleARCoreBase/Private/GoogleARCoreBaseModule.cpp b/Engine/Plugins/Runtime/AR/Google/GoogleARCore/Source/GoogleARCoreBase/Private/GoogleARCoreBaseModule.cpp
--- a/Engine/Plugins/Runtime/AR/Google/GoogleARCore/Source/GoogleARCoreBase/Private/GoogleARCoreBaseModule.cpp
+++ b/Engine/Plugins/Runtime/AR/Google/GoogleARCore/Source/GoogleARCoreBase/Private/GoogleARCoreBaseModule.cpp
@@ -11,9 +11,14 @@
 #include "GoogleARCoreXRTrackingSystem.h"
 #include "GoogleARCoreDevice.h"
 #include "GoogleARCoreCookSupport.h"
+#include "GoogleARCoreBaseLogCategory.h"
 
 #define LOCTEXT_NAMESPACE "GoogleARCore"
 
+#if PLATFORM_ANDROID
+#include "IVulkanDynamicRHI.h"
+#endif
+
 class FGoogleARCoreBaseModule : public IGoogleARCoreBaseModule
 {
 public:
@@ -68,6 +73,26 @@ TSharedPtr< class IXRTrackingSystem, ESPMode::ThreadSafe > FGoogleARCoreBaseModu
 
 void FGoogleARCoreBaseModule::StartupModule()
 {
+#if PLATFORM_ANDROID
+	if (FAndroidMisc::ShouldUseVulkan())
+	{
+		IVulkanDynamicRHI* RHI = GetIVulkanDynamicRHI();
+		TArray<FAnsiString> LoadedDeviceExtensions = RHI->RHIGetLoadedDeviceExtensions();
+		if (!LoadedDeviceExtensions.Find(FAnsiString(VK_ANDROID_EXTERNAL_MEMORY_ANDROID_HARDWARE_BUFFER_EXTENSION_NAME)))
+		{
+			UE_LOG(LogGoogleARCore, Error, TEXT("Required Vulkan Device extension '" VK_ANDROID_EXTERNAL_MEMORY_ANDROID_HARDWARE_BUFFER_EXTENSION_NAME "' is not available!"));
+		}
+		if (!LoadedDeviceExtensions.Find(FAnsiString(VK_EXT_QUEUE_FAMILY_FOREIGN_EXTENSION_NAME)))
+		{
+			UE_LOG(LogGoogleARCore, Error, TEXT("Required Vulkan Device extension '" VK_EXT_QUEUE_FAMILY_FOREIGN_EXTENSION_NAME "' is not available!"));
+		}
+		if (!LoadedDeviceExtensions.Find(FAnsiString(VK_KHR_SAMPLER_YCBCR_CONVERSION_EXTENSION_NAME)))
+		{
+			UE_LOG(LogGoogleARCore, Error, TEXT("Required Vulkan Device extension '" VK_KHR_SAMPLER_YCBCR_CONVERSION_EXTENSION_NAME "' is not available!"));
+		}
+	}
+#endif
+
 	ensureMsgf(FModuleManager::Get().LoadModule("AugmentedReality"), TEXT("ARCore depends on the AugmentedReality module.") );
 
 	// Register editor settings:
diff --git a/Engine/Plugins/Runtime/AR/Google/GoogleARCore/Source/GoogleARCoreBase/Private/GoogleARCoreCameraTextureResourceVulkan.cpp b/Engine/Plugins/Runtime/AR/Google/GoogleARCore/Source/GoogleARCoreBase/Private/GoogleARCoreCameraTextureResourceVulkan.cpp
new file mode 100644
--- /dev/null
+++ b/Engine/Plugins/Runtime/AR/Google/GoogleARCore/Source/GoogleARCoreBase/Private/GoogleARCoreCameraTextureResourceVulkan.cpp
@@ -0,0 +1,111 @@
+// Copyright Epic Games, Inc. All Rights Reserved.
+
+#include "GoogleARCoreCameraTextureResourceVulkan.h"
+#include "GoogleARCoreBaseLogCategory.h"
+#include "ExternalTexture.h"
+
+#if PLATFORM_ANDROID
+#include "VulkanCommon.h"
+#include "IVulkanDynamicRHI.h"
+#endif
+
+FARCoreCameraTextureResourceVulkan::FARCoreCameraTextureResourceVulkan(const FGuid& InExternalTextureGuid, void* InHardwareBuffer, uint32 InSizeX, uint32 InSizeY)
+	: ExternalTextureGuid(InExternalTextureGuid)
+	, SizeX(InSizeX)
+	, SizeY(InSizeY)
+	, HardwareBuffer(InHardwareBuffer)
+{
+	// With Vulkan we don't use the Id from the native resource because it gets recreated every frame.
+	// We handle the binding of the hardware buffer to the texture ourselves. Using a hash of the external
+	// texture Guid is sufficient and will remain the same from frame to frame.
+	TextureId = GetTypeHash(ExternalTextureGuid);
+}
+
+void FARCoreCameraTextureResourceVulkan::InitRHI(FRHICommandListBase& RHICmdList)
+{
+	if (HardwareBuffer)
+	{
+		FSamplerStateInitializerRHI SamplerStateInitializer(SF_Bilinear, AM_Clamp, AM_Clamp, AM_Clamp);
+		SamplerStateRHI = RHICreateSamplerState(SamplerStateInitializer);
+
+#if PLATFORM_ANDROID
+		IVulkanDynamicRHI* VulkanRHI = GetIVulkanDynamicRHI();
+		TextureRHI = VulkanRHI->RHICreateTexture2DFromAndroidHardwareBuffer(reinterpret_cast<AHardwareBuffer*>(HardwareBuffer));
+#endif
+
+		FillCameraYCbCrConversionParameters();
+
+		FExternalTextureRegistry::Get().RegisterExternalTexture(ExternalTextureGuid, TextureRHI, SamplerStateRHI);
+
+	}
+}
+
+void FARCoreCameraTextureResourceVulkan::ReleaseRHI()
+{
+	if (HardwareBuffer)
+	{
+		FExternalTextureRegistry::Get().UnregisterExternalTexture(ExternalTextureGuid);
+		FTextureResource::ReleaseRHI();
+	}
+}
+
+#if PLATFORM_ANDROID
+namespace VulkanUtil
+{
+	EYCbCrModelConversion ToEYCbCrModelConversion(VkSamplerYcbcrModelConversion VulkanYcbcrModelConversion)
+	{
+		switch (VulkanYcbcrModelConversion)
+		{
+		case VK_SAMPLER_YCBCR_MODEL_CONVERSION_RGB_IDENTITY:   return EYCbCrModelConversion::None;
+		case VK_SAMPLER_YCBCR_MODEL_CONVERSION_YCBCR_IDENTITY: return EYCbCrModelConversion::YCbCrIdentity;
+		case VK_SAMPLER_YCBCR_MODEL_CONVERSION_YCBCR_709:      return EYCbCrModelConversion::YCbCrRec709;
+		case VK_SAMPLER_YCBCR_MODEL_CONVERSION_YCBCR_601:      return EYCbCrModelConversion::YCbCrRec601;
+		case VK_SAMPLER_YCBCR_MODEL_CONVERSION_YCBCR_2020:     return EYCbCrModelConversion::YCbCrRec2020;
+
+		default:
+			UE_LOG(LogGoogleARCore, Error, TEXT("Unexpected Vulkan Ycbcr Model Conversion value: %d"), VulkanYcbcrModelConversion);
+			return EYCbCrModelConversion::None;
+		}
+	}
+
+	EYCbCrRange ToEYCbCrRange(VkSamplerYcbcrRange VulkanYcbcrRange)
+	{
+		switch (VulkanYcbcrRange)
+		{
+		case VK_SAMPLER_YCBCR_RANGE_ITU_FULL:   return EYCbCrRange::Full;
+		case VK_SAMPLER_YCBCR_RANGE_ITU_NARROW: return EYCbCrRange::Narrow;
+
+		default:
+			UE_LOG(LogGoogleARCore, Error, TEXT("Unexpected Vulkan Ycbcr Range value: %d"), VulkanYcbcrRange);
+			return EYCbCrRange::Unknown;
+		}
+	}
+}
+#endif
+
+void FARCoreCameraTextureResourceVulkan::FillCameraYCbCrConversionParameters()
+{
+#if PLATFORM_ANDROID
+	IVulkanDynamicRHI* VulkanRHI = GetIVulkanDynamicRHI();
+	VkDevice VulkanDevice = VulkanRHI->RHIGetVkDevice();
+
+	VkAndroidHardwareBufferFormatPropertiesANDROID HardwareBufferFormatProperties;
+	ZeroVulkanStruct(HardwareBufferFormatProperties, VK_STRUCTURE_TYPE_ANDROID_HARDWARE_BUFFER_FORMAT_PROPERTIES_ANDROID);
+
+	VkAndroidHardwareBufferPropertiesANDROID HardwareBufferProperties;
+	ZeroVulkanStruct(HardwareBufferProperties, VK_STRUCTURE_TYPE_ANDROID_HARDWARE_BUFFER_PROPERTIES_ANDROID);
+	HardwareBufferProperties.pNext = &HardwareBufferFormatProperties;
+
+	PFN_vkGetAndroidHardwareBufferPropertiesANDROID GetAndroidHardwareBufferPropertiesANDROID = (PFN_vkGetAndroidHardwareBufferPropertiesANDROID)VulkanRHI->RHIGetVkDeviceProcAddr("vkGetAndroidHardwareBufferPropertiesANDROID");
+	if (!GetAndroidHardwareBufferPropertiesANDROID)
+	{
+		UE_LOG(LogGoogleARCore, Error, TEXT("PFN_vkGetAndroidHardwareBufferPropertiesANDROID address not found."));
+		return;
+	}
+	VERIFYVULKANRESULT_EXTERNAL(GetAndroidHardwareBufferPropertiesANDROID(VulkanDevice, reinterpret_cast<AHardwareBuffer*>(HardwareBuffer), &HardwareBufferProperties));
+
+	CameraYCbCrConversion.YCbCrModelConversion = VulkanUtil::ToEYCbCrModelConversion(HardwareBufferFormatProperties.suggestedYcbcrModel);
+	CameraYCbCrConversion.YCbCrRange = VulkanUtil::ToEYCbCrRange(HardwareBufferFormatProperties.suggestedYcbcrRange);
+	CameraYCbCrConversion.NumBits = 8; // Defaulting to 8 bits per component. Investigate how to obtain it from HardwareBuffer.
+#endif
+}
diff --git a/Engine/Plugins/Runtime/AR/Google/GoogleARCore/Source/GoogleARCoreBase/Private/GoogleARCoreCameraTextureResourceVulkan.h b/Engine/Plugins/Runtime/AR/Google/GoogleARCore/Source/GoogleARCoreBase/Private/GoogleARCoreCameraTextureResourceVulkan.h
new file mode 100644
--- /dev/null
+++ b/Engine/Plugins/Runtime/AR/Google/GoogleARCore/Source/GoogleARCoreBase/Private/GoogleARCoreCameraTextureResourceVulkan.h
@@ -0,0 +1,61 @@
+// Copyright Epic Games, Inc. All Rights Reserved.
+
+#pragma once
+
+#include "GoogleARCoreYCbCrConversion.h"
+#include "TextureResource.h"
+#include "RHICommandList.h"
+#include "Misc/Guid.h"
+
+// Camera texture resource using the Vulkan hardware buffer provided.
+// 
+// Known Vulkan Validation error: VUID-VkWriteDescriptorSet-descriptorType-01946
+//
+// Vulkan RHI will create the image and its view with a VkSamplerYcbcrConversion, and bind it as a VK_DESCRIPTOR_TYPE_SAMPLED_IMAGE.
+// But using a VkSamplerYcbcrConversion requires the sampler and the image view to be bound together with a
+// VK_DESCRIPTOR_TYPE_COMBINED_IMAGE_SAMPLER, which is not supported.
+// 
+// As a workaround we perform the YCbCr conversion ourselves in the shader and do not pass VkSamplerYcbcrConversion
+// to the sampler. The VkSamplerYcbcrConversion is still necessary for the view, which reports the validation error,
+// and it will read the values from the external image. In the future if support for combined image sampler descriptors
+// is provided, then it would possible to use VkSamplerYcbcrConversion natively and remove the YCbCr conversion inside
+// the shader.
+class FARCoreCameraTextureResourceVulkan : public FTextureResource
+{
+public:
+	FARCoreCameraTextureResourceVulkan(const FGuid& InExternalTextureGuid, void* InHardwareBuffer, uint32 InSizeX, uint32 InSizeY);
+
+	virtual void InitRHI(FRHICommandListBase& RHICmdList) override;
+	virtual void ReleaseRHI() override;
+
+	virtual uint32 GetSizeX() const override
+	{
+		return SizeX;
+	}
+
+	virtual uint32 GetSizeY() const override
+	{
+		return SizeY;
+	}
+
+	uint32 GetTextureId() const
+	{
+		return TextureId;
+	}
+
+	FYCbCrConversion GetCameraYCbCrConversion() const
+	{
+		return CameraYCbCrConversion;
+	}
+
+private:
+	void FillCameraYCbCrConversionParameters();
+
+	uint32 TextureId = 0;
+	const FGuid ExternalTextureGuid;
+	const uint32 SizeX;
+	const uint32 SizeY;
+
+	void* HardwareBuffer = nullptr;
+	FYCbCrConversion CameraYCbCrConversion;
+};
diff --git a/Engine/Plugins/Runtime/AR/Google/GoogleARCore/Source/GoogleARCoreBase/Private/GoogleARCoreDevice.cpp b/Engine/Plugins/Runtime/AR/Google/GoogleARCore/Source/GoogleARCoreBase/Private/GoogleARCoreDevice.cpp
--- a/Engine/Plugins/Runtime/AR/Google/GoogleARCore/Source/GoogleARCoreBase/Private/GoogleARCoreDevice.cpp
+++ b/Engine/Plugins/Runtime/AR/Google/GoogleARCore/Source/GoogleARCoreBase/Private/GoogleARCoreDevice.cpp
@@ -24,15 +24,6 @@ DECLARE_CYCLE_STAT(TEXT("UpdateGameFrame"), STAT_UpdateGameFrame, STATGROUP_ARCo
 
 #define NUM_CAMERA_TEXTURES 4
 
-int32 GARCoreEnableVulkanSupport = 0;
-static FAutoConsoleVariableRef CVarARCoreEnableVulkanSupport(
-	TEXT("r.ARCore.EnableVulkanSupport"),
-	GARCoreEnableVulkanSupport,
-	TEXT("Whether to support Vulkan in the ARCore plugin.\n")
-	TEXT(" 0: Disabled (default)\n")
-	TEXT(" 1: Enabled"),
-	ECVF_Default);
-
 namespace
 {
 	EGoogleARCoreFunctionStatus ToARCoreFunctionStatus(EGoogleARCoreAPIStatus Status)
@@ -148,12 +139,7 @@ void FGoogleARCoreDevice::StartARCoreSessionRequest(UARSessionConfig* SessionCon
 	}
 	
 #if PLATFORM_ANDROID
-	// Vulkan support is currently disabled, until it's revisited and fixed
-	if (ShouldUseVulkan()) // !GARCoreEnableVulkanSupport
-	{
-		UE_LOG(LogGoogleARCore, Error, TEXT("StartARCoreSessionRequest ignored - ARCore requires an OpenGL context, but we're using Vulkan!"));
-		return;
-	}
+	UE_LOG(LogGoogleARCore, Display, TEXT("ARCore session using %s."), ShouldUseVulkan() ? TEXT("Vulkan") : TEXT("OpenGL ES"));
 #endif
 
 	if (bIsARCoreSessionRunning)
@@ -357,11 +343,27 @@ void FGoogleARCoreDevice::UpdateGameFrame(UWorld* World)
 			CurrentSessionStatus.Status = EARSessionStatus::FatalError;
 			CurrentSessionStatus.AdditionalInfo = TEXT("Fatal error occurred when updating ARCore Session. Stopping and restarting ARCore Session may fix the issue.");
 		}
-		else
+		else if (ARCoreSession->HasNewFrame()) // Only update textures if the session has a new frame
 		{
 			if (auto Frame = ARCoreSession->GetLatestFrame())
 			{
-				LastCameraTextureId = Frame->GetCameraTextureId();
+				// With Vulkan update the camera texture with the current frame's hardware buffer.
+				if (ShouldUseVulkan())
+				{
+					const uint32 CurrentFrame = ARCoreSession->GetFrameNum() % PassthroughCameraTextures.Num();
+					const uint32 CurrentCameraTextureId = PassthroughCameraTextureKeys[CurrentFrame];
+
+					if (Frame->UpdateCameraTexture(PassthroughCameraTextures[CurrentCameraTextureId]))
+					{
+						LastCameraTextureId = CurrentCameraTextureId;
+					}
+				}
+				// With OpenGL just use the texture id returned by the frame.
+				else
+				{
+					LastCameraTextureId = Frame->GetCameraTextureId();
+				}
+
 				if (ARCoreSession->IsSceneDepthEnabled())
 				{
 					Frame->UpdateDepthTexture(MutableView(DepthTexture));
@@ -429,11 +431,6 @@ void FGoogleARCoreDevice::HandleRuntimePermissionsGranted(const TArray<FString>&
 void FGoogleARCoreDevice::StartSessionWithRequestedConfig()
 {
 	bStartSessionRequested = false;
-	
-	if (ShouldUseVulkan())
-	{
-		GLContext = FGoogleARCoreOpenGLContext::CreateContext();
-	}
 
 	ARCoreSession = CreateSession();
 
@@ -604,7 +601,7 @@ void FGoogleARCoreDevice::StartSession()
 	check(PassthroughCameraTextures.Num());
 	TArray<uint32> TextureIds;
 	PassthroughCameraTextures.GetKeys(TextureIds);
-	UE_LOG(LogGoogleARCore, Error, TEXT("SetCameraTextureIds %i textures"), TextureIds.Num());
+	UE_LOG(LogGoogleARCore, Display, TEXT("SetCameraTextureIds %i textures"), TextureIds.Num());
 	ARCoreSession->SetCameraTextureIds(TextureIds);
 
 
@@ -711,12 +708,9 @@ void FGoogleARCoreDevice::ResetARCoreSession()
 	ARCoreSession.Reset();
 	CurrentSessionStatus.Status = EARSessionStatus::NotStarted;
 	CurrentSessionStatus.AdditionalInfo = TEXT("ARCore Session is uninitialized.");
-	if (GLContext)
-	{
-		GLContext = nullptr;
-	}
 	
 	PassthroughCameraTextures = {};
+	PassthroughCameraTextureKeys = {};
 	LastCameraTextureId = 0;
 	UE_LOG(LogGoogleARCore, Log, TEXT("ARCore session reset"));
 }
@@ -747,6 +741,9 @@ void FGoogleARCoreDevice::AllocatePassthroughCameraTextures()
 		UE_LOG(LogGoogleARCore, Log, TEXT("Created external camera texture of size (%.0f x %.0f) with texture Id (%d) and external texture Guid (%s)"),
 			   Texture->Size.X, Texture->Size.Y, TextureId, *Texture->ExternalTextureGuid.ToString());
 	}
+
+	// Keep an array of the keys to be able circle through the textures frame by frame when using Vulkan.
+	PassthroughCameraTextures.GetKeys(PassthroughCameraTextureKeys);
 }
 
 FMatrix FGoogleARCoreDevice::GetPassthroughCameraProjectionMatrix(FIntPoint ViewRectSize) const
@@ -1022,6 +1019,16 @@ UARTexture* FGoogleARCoreDevice::GetDepthTexture() const
 	return DepthTexture;
 }
 
+IYCbCrConversionQuery* FGoogleARCoreDevice::GetLastCameraYCbCrConversionQuery() const
+{
+	if (auto Record = PassthroughCameraTextures.Find(LastCameraTextureId))
+	{
+		return *Record;
+	}
+
+	return nullptr;
+}
+
 void FGoogleARCoreDevice::FInternalGCObject::AddReferencedObjects(FReferenceCollector& Collector)
 {
 	Collector.AddReferencedObject(ARCoreDevice->DepthTexture);
diff --git a/Engine/Plugins/Runtime/AR/Google/GoogleARCore/Source/GoogleARCoreBase/Private/GoogleARCoreDevice.h b/Engine/Plugins/Runtime/AR/Google/GoogleARCore/Source/GoogleARCoreBase/Private/GoogleARCoreDevice.h
--- a/Engine/Plugins/Runtime/AR/Google/GoogleARCore/Source/GoogleARCoreBase/Private/GoogleARCoreDevice.h
+++ b/Engine/Plugins/Runtime/AR/Google/GoogleARCore/Source/GoogleARCoreBase/Private/GoogleARCoreDevice.h
@@ -13,7 +13,6 @@
 #include "GoogleARCoreAPI.h"
 #include "GoogleARCorePassthroughCameraRenderer.h"
 #include "GoogleARCoreAugmentedImageDatabase.h"
-#include "GoogleARCoreOpenGLContext.h"
 
 
 class UARTexture;
@@ -153,6 +152,8 @@ public:
 	
 	UARTexture* GetLastCameraTexture() const;
 	UARTexture* GetDepthTexture() const;
+
+	IYCbCrConversionQuery* GetLastCameraYCbCrConversionQuery() const;
 	
 	EARCorePermissionStatus CheckAndRequrestPermission(const UARSessionConfig& ConfigurationData);
 
@@ -184,6 +185,7 @@ private:
 	TSharedPtr<FGoogleARCoreSession> BackCameraARCoreSession;
 	
 	TMap<uint32, TObjectPtr<UARCoreCameraTexture>> PassthroughCameraTextures;
+	TArray<uint32> PassthroughCameraTextureKeys;
 	
 	TObjectPtr<UARCoreDepthTexture> DepthTexture;
 	
@@ -213,8 +215,6 @@ private:
 
 	TSharedPtr<FARSupportInterface, ESPMode::ThreadSafe> ARSystem;
 	
-	TSharedPtr<FGoogleARCoreOpenGLContext, ESPMode::ThreadSafe> GLContext;
-	
 	// FGoogleARCoreDevice itself cannot be a FGCObject since it's initialize too early.
 	// So we use an internal class instead, which can be created only when it's needed.
 	class FInternalGCObject : FGCObject
diff --git a/Engine/Plugins/Runtime/AR/Google/GoogleARCore/Source/GoogleARCoreBase/Private/GoogleARCoreOpenGLContext.cpp b/Engine/Plugins/Runtime/AR/Google/GoogleARCore/Source/GoogleARCoreBase/Private/GoogleARCoreOpenGLContext.cpp
deleted file mode 100644
diff --git a/Engine/Plugins/Runtime/AR/Google/GoogleARCore/Source/GoogleARCoreBase/Private/GoogleARCoreOpenGLContext.h b/Engine/Plugins/Runtime/AR/Google/GoogleARCore/Source/GoogleARCoreBase/Private/GoogleARCoreOpenGLContext.h
deleted file mode 100644
diff --git a/Engine/Plugins/Runtime/AR/Google/GoogleARCore/Source/GoogleARCoreBase/Private/GoogleARCoreTexture.cpp b/Engine/Plugins/Runtime/AR/Google/GoogleARCore/Source/GoogleARCoreBase/Private/GoogleARCoreTexture.cpp
--- a/Engine/Plugins/Runtime/AR/Google/GoogleARCore/Source/GoogleARCoreBase/Private/GoogleARCoreTexture.cpp
+++ b/Engine/Plugins/Runtime/AR/Google/GoogleARCore/Source/GoogleARCoreBase/Private/GoogleARCoreTexture.cpp
@@ -1,6 +1,7 @@
 // Copyright Epic Games, Inc. All Rights Reserved.
 
 #include "GoogleARCoreTexture.h"
+#include "GoogleARCoreCameraTextureResourceVulkan.h"
 #include "TextureResource.h"
 #include "RHICommandList.h"
 #include "ExternalTexture.h"
@@ -70,11 +71,29 @@ UARCoreCameraTexture::UARCoreCameraTexture(const FObjectInitializer& ObjectIniti
 
 FTextureResource* UARCoreCameraTexture::CreateResource()
 {
+#if PLATFORM_ANDROID
+	if (FAndroidMisc::ShouldUseVulkan())
+	{
+		return new FARCoreCameraTextureResourceVulkan(ExternalTextureGuid, VukanHardwareBuffer, Size.X, Size.Y);
+	}
+#endif
+
 	return new FARCoreCameraTextureResource(ExternalTextureGuid, Size.X, Size.Y);
 }
 
 uint32 UARCoreCameraTexture::GetTextureId() const
 {
+#if PLATFORM_ANDROID
+	if (FAndroidMisc::ShouldUseVulkan())
+	{
+		if (const auto* MyResource = static_cast<const FARCoreCameraTextureResourceVulkan*>(GetResource()))
+		{
+			return MyResource->GetTextureId();
+		}
+		return 0;
+	}
+#endif
+
 	if (const FARCoreCameraTextureResource* MyResource = static_cast<const FARCoreCameraTextureResource*>(GetResource()))
 	{
 		return MyResource->GetTextureId();
@@ -83,6 +102,29 @@ uint32 UARCoreCameraTexture::GetTextureId() const
 	return 0;
 }
 
+void UARCoreCameraTexture::UpdateCameraImage(void* InVukanHardwareBuffer)
+{
+	if (InVukanHardwareBuffer)
+	{
+		VukanHardwareBuffer = InVukanHardwareBuffer;
+		UpdateResource();
+	}
+}
+
+FYCbCrConversion UARCoreCameraTexture::GetYCbCrConversion_RenderThread() const
+{
+#if PLATFORM_ANDROID
+	if (FAndroidMisc::ShouldUseVulkan())
+	{
+		if (const auto* MyResource = static_cast<const FARCoreCameraTextureResourceVulkan*>(GetResource()))
+		{
+			return MyResource->GetCameraYCbCrConversion();
+		}
+	}
+#endif
+
+	return {};
+}
 
 class FARCoreDepthTextureResource : public FTextureResource
 {
diff --git a/Engine/Plugins/Runtime/AR/Google/GoogleARCore/Source/GoogleARCoreBase/Private/GoogleARCoreTexture.h b/Engine/Plugins/Runtime/AR/Google/GoogleARCore/Source/GoogleARCoreBase/Private/GoogleARCoreTexture.h
--- a/Engine/Plugins/Runtime/AR/Google/GoogleARCore/Source/GoogleARCoreBase/Private/GoogleARCoreTexture.h
+++ b/Engine/Plugins/Runtime/AR/Google/GoogleARCore/Source/GoogleARCoreBase/Private/GoogleARCoreTexture.h
@@ -2,6 +2,7 @@
 
 #pragma once
 
+#include "GoogleARCoreYCbCrConversion.h"
 #include "ARTextures.h"
 #include "Misc/Guid.h"
 
@@ -14,6 +15,7 @@
 
 UCLASS()
 class GOOGLEARCOREBASE_API UARCoreCameraTexture : public UARTexture
+												, public IYCbCrConversionQuery
 {
 	GENERATED_UCLASS_BODY()
 	
@@ -24,6 +26,15 @@ public:
 	// End UTexture interface
 	
 	uint32 GetTextureId() const;
+
+	// Recreates image resource using Vulkan's external Hardware Buffer.
+	void UpdateCameraImage(void* InVulkanHardwareBuffer);
+
+	// IYCbCrConversionQuery overrides
+	FYCbCrConversion GetYCbCrConversion_RenderThread() const override;
+
+private:
+	void* VukanHardwareBuffer = nullptr;
 };
 
 
diff --git a/Engine/Plugins/Runtime/AR/Google/GoogleARCore/Source/GoogleARCoreBase/Private/GoogleARCoreXRCamera.cpp b/Engine/Plugins/Runtime/AR/Google/GoogleARCore/Source/GoogleARCoreBase/Private/GoogleARCoreXRCamera.cpp
--- a/Engine/Plugins/Runtime/AR/Google/GoogleARCore/Source/GoogleARCoreBase/Private/GoogleARCoreXRCamera.cpp
+++ b/Engine/Plugins/Runtime/AR/Google/GoogleARCore/Source/GoogleARCoreBase/Private/GoogleARCoreXRCamera.cpp
@@ -51,13 +51,17 @@ void FGoogleARCoreXRCamera::PreRenderViewFamily_RenderThread(FRDGBuilder& GraphB
 	}
 
 #if PLATFORM_ANDROID
-	if(TS.ARCoreDeviceInstance->GetIsARCoreSessionRunning() && TS.ARCoreDeviceInstance->GetShouldInvertCulling())
+	// Culling inversion ignored in Vulkan. Investigate when it's required and approach to apply it in Vulkan.
+	if (!FAndroidMisc::ShouldUseVulkan())
 	{
-		glFrontFace(GL_CW);
-	}
-	else
-	{
-		glFrontFace(GL_CCW);
+		if(TS.ARCoreDeviceInstance->GetIsARCoreSessionRunning() && TS.ARCoreDeviceInstance->GetShouldInvertCulling())
+		{
+			glFrontFace(GL_CW);
+		}
+		else
+		{
+			glFrontFace(GL_CCW);
+		}
 	}
 #endif
 }
@@ -143,3 +147,11 @@ void FGoogleARCoreXRCamera::UpdateCameraTextures(UTexture* NewCameraTexture, UTe
 		PassthroughRenderer->UpdateCameraTextures(NewCameraTexture, DepthTexture, bEnableOcclusion);
 	}
 }
+
+void FGoogleARCoreXRCamera::UpdateCameraYCbCrConversion(IYCbCrConversionQuery* NewYCbCrConversionQuery)
+{
+	if (PassthroughRenderer)
+	{
+		PassthroughRenderer->UpdateCameraYCbCrConversion(NewYCbCrConversionQuery);
+	}
+}
diff --git a/Engine/Plugins/Runtime/AR/Google/GoogleARCore/Source/GoogleARCoreBase/Private/GoogleARCoreXRCamera.h b/Engine/Plugins/Runtime/AR/Google/GoogleARCore/Source/GoogleARCoreBase/Private/GoogleARCoreXRCamera.h
--- a/Engine/Plugins/Runtime/AR/Google/GoogleARCore/Source/GoogleARCoreBase/Private/GoogleARCoreXRCamera.h
+++ b/Engine/Plugins/Runtime/AR/Google/GoogleARCore/Source/GoogleARCoreBase/Private/GoogleARCoreXRCamera.h
@@ -7,6 +7,7 @@
 class FGoogleARCoreXRTrackingSystem;
 class FGoogleARCorePassthroughCameraRenderer;
 class UTexture;
+class IYCbCrConversionQuery;
 
 class FGoogleARCoreXRCamera : public FDefaultXRCamera
 {
@@ -28,6 +29,8 @@ public:
 	
 	void UpdateCameraTextures(UTexture* NewCameraTexture, UTexture* DepthTexture, bool bEnableOcclusion);
 
+	void UpdateCameraYCbCrConversion(IYCbCrConversionQuery* NewYCbCrConversionQuery);
+
 private:
 	FGoogleARCoreXRTrackingSystem& GoogleARCoreTrackingSystem;
 	FGoogleARCorePassthroughCameraRenderer* PassthroughRenderer;
diff --git a/Engine/Plugins/Runtime/AR/Google/GoogleARCore/Source/GoogleARCoreBase/Private/GoogleARCoreXRTrackingSystem.cpp b/Engine/Plugins/Runtime/AR/Google/GoogleARCore/Source/GoogleARCoreBase/Private/GoogleARCoreXRTrackingSystem.cpp
--- a/Engine/Plugins/Runtime/AR/Google/GoogleARCore/Source/GoogleARCoreBase/Private/GoogleARCoreXRTrackingSystem.cpp
+++ b/Engine/Plugins/Runtime/AR/Google/GoogleARCore/Source/GoogleARCoreBase/Private/GoogleARCoreXRTrackingSystem.cpp
@@ -153,6 +153,9 @@ bool FGoogleARCoreXRTrackingSystem::OnStartGameFrame(FWorldContext& WorldContext
 		
 		if (auto Camera = GetARCoreCamera())
 		{
+			auto CameraYCbCrConversionQuery = ARCoreDeviceInstance->GetLastCameraYCbCrConversionQuery();
+			Camera->UpdateCameraYCbCrConversion(CameraYCbCrConversionQuery);
+
 			auto CameraTexture = ARCoreDeviceInstance->GetLastCameraTexture();
 			auto DepthTexture = ARCoreDeviceInstance->GetDepthTexture();
 			Camera->UpdateCameraTextures(CameraTexture, DepthTexture, bWantsDepthOcclusion);
diff --git a/Engine/Plugins/Runtime/AR/Google/GoogleARCore/Source/GoogleARCoreRendering/Private/GoogleARCorePassthroughCameraRenderer.cpp b/Engine/Plugins/Runtime/AR/Google/GoogleARCore/Source/GoogleARCoreRendering/Private/GoogleARCorePassthroughCameraRenderer.cpp
--- a/Engine/Plugins/Runtime/AR/Google/GoogleARCore/Source/GoogleARCoreRendering/Private/GoogleARCorePassthroughCameraRenderer.cpp
+++ b/Engine/Plugins/Runtime/AR/Google/GoogleARCore/Source/GoogleARCoreRendering/Private/GoogleARCorePassthroughCameraRenderer.cpp
@@ -24,6 +24,7 @@
 #include "PostProcess/DrawRectangle.h"
 #include "RHIResourceUtils.h"
 #include "RHIStaticStates.h"
+#include "MediaShaders.h"
 
 #if !(UE_BUILD_SHIPPING || UE_BUILD_TEST)
 	#define ALLOWS_DEBUG_OVERLAY 1
@@ -163,7 +164,87 @@ public:
 
 IMPLEMENT_GLOBAL_SHADER(FGoogleARCoreCameraOverlayPS, "/Engine/Private/PostProcessMaterialShaders.usf", "MainPS", SF_Pixel);
 
-void FGoogleARCorePassthroughCameraRenderer::RenderVideoOverlayWithMaterial(FRHICommandList& RHICmdList, FSceneView& InView, UMaterialInstanceDynamic* OverlayMaterialToUse, bool bRenderingOcclusion)
+// Postprocess pixel shader for ARCore camera performing YCbCr conversion.
+class FGoogleARCoreCameraOverlayYCbCrConversionPS : public FPostProcessMaterialShader
+{
+public:
+	DECLARE_SHADER_TYPE(FGoogleARCoreCameraOverlayYCbCrConversionPS, Material);
+
+	SHADER_USE_PARAMETER_STRUCT_WITH_LEGACY_BASE(FGoogleARCoreCameraOverlayYCbCrConversionPS, FPostProcessMaterialShader)
+
+	BEGIN_SHADER_PARAMETER_STRUCT(FYCbCrConversionParameters, )
+		SHADER_PARAMETER(FMatrix44f, YCbCrColorTransform)
+		SHADER_PARAMETER(uint32, YCbCrSrgbToLinear)
+	END_SHADER_PARAMETER_STRUCT()
+	using FParameters = FYCbCrConversionParameters;
+
+	static void ModifyCompilationEnvironment(const FMaterialShaderPermutationParameters& Parameters, FShaderCompilerEnvironment& OutEnvironment)
+	{
+		FPostProcessMaterialShader::ModifyCompilationEnvironment(Parameters, OutEnvironment);
+		OutEnvironment.SetDefine(TEXT("OUTPUT_GAMMA_SPACE"), IsMobileHDR() ? 0 : 1);
+		OutEnvironment.SetDefine(TEXT("POST_PROCESS_AR_PASSTHROUGH"), 1);
+		OutEnvironment.SetDefine(TEXT("POST_PROCESS_AR_YCBCR_CONVERSION"), 1);
+	}
+
+	void SetParameters(FRHIBatchedShaderParameters& BatchedParameters, const FSceneView& View, const FMaterialRenderProxy* MaterialProxy, const FMaterial& Material)
+	{
+		FMaterialShader::SetViewParameters(BatchedParameters, View, View.ViewUniformBuffer);
+		FMaterialShader::SetParameters(BatchedParameters, MaterialProxy, Material, View);
+	}
+
+	FYCbCrConversionParameters GetYCbCrConversionParameters(const FYCbCrConversion& YCbCrConversion) const
+	{
+		FMatrix YCbCrModelConversionMatrix = FMatrix::Identity;
+		FVector YCbCrRangeOffset = FVector::ZeroVector;
+
+		check(YCbCrConversion.YCbCrRange != EYCbCrRange::Unknown);
+		const bool bFullRange = (YCbCrConversion.YCbCrRange == EYCbCrRange::Full);
+
+		switch (YCbCrConversion.YCbCrModelConversion)
+		{
+		case EYCbCrModelConversion::YCbCrIdentity:
+			// Use identity matrix
+			break;
+
+		case EYCbCrModelConversion::YCbCrRec709:
+			YCbCrModelConversionMatrix = (bFullRange) ? MediaShaders::YuvToRgbRec709Unscaled : MediaShaders::YuvToRgbRec709Scaled;
+			break;
+
+		case EYCbCrModelConversion::YCbCrRec601:
+			YCbCrModelConversionMatrix = (bFullRange) ? MediaShaders::YuvToRgbRec601Unscaled : MediaShaders::YuvToRgbRec601Scaled;
+			break;
+
+		case EYCbCrModelConversion::YCbCrRec2020:
+			YCbCrModelConversionMatrix = (bFullRange) ? MediaShaders::YuvToRgbRec2020Unscaled : MediaShaders::YuvToRgbRec2020Scaled;
+			break;
+
+		case EYCbCrModelConversion::None: // None value is not expected, otherwise it should be using FGoogleARCoreCameraOverlayPS shader instead.
+		default:
+			UE_LOG(LogTemp, Error, TEXT("Unexpected YCbCr Model conversion: %d"), YCbCrConversion.YCbCrModelConversion);
+			break;
+		}
+
+		switch (YCbCrConversion.NumBits)
+		{
+		case 8:  YCbCrRangeOffset = bFullRange ? MediaShaders::YUVOffsetNoScale8bits : MediaShaders::YUVOffset8bits; break;
+		case 10: YCbCrRangeOffset = bFullRange ? MediaShaders::YUVOffsetNoScale10bits : MediaShaders::YUVOffset10bits; break;
+		case 16: YCbCrRangeOffset = bFullRange ? MediaShaders::YUVOffsetNoScale16bits : MediaShaders::YUVOffset16bits; break;
+		case 32: YCbCrRangeOffset = bFullRange ? MediaShaders::YUVOffsetNoScaleFloat : MediaShaders::YUVOffsetFloat; break;
+		default:
+			UE_LOG(LogTemp, Error, TEXT("Unexpected number of bits in YCbCr conversion: %d"), YCbCrConversion.NumBits);
+			break;
+		}
+
+		FYCbCrConversionParameters YCbCrConversionParameters;
+		YCbCrConversionParameters.YCbCrColorTransform = (FMatrix44f)MediaShaders::CombineColorTransformAndOffset(YCbCrModelConversionMatrix, YCbCrRangeOffset);
+		YCbCrConversionParameters.YCbCrSrgbToLinear = 1;
+		return YCbCrConversionParameters;
+	}
+};
+
+IMPLEMENT_GLOBAL_SHADER(FGoogleARCoreCameraOverlayYCbCrConversionPS, "/Engine/Private/PostProcessMaterialShaders.usf", "MainPS", SF_Pixel);
+
+void FGoogleARCorePassthroughCameraRenderer::RenderVideoOverlayWithMaterial(FRHICommandList& RHICmdList, FSceneView& InView, UMaterialInstanceDynamic* OverlayMaterialToUse, bool bRenderingOcclusion, bool bUsesCameraTexture)
 {
 #if PLATFORM_ANDROID
 	if (FAndroidMisc::ShouldUseVulkan() && IsMobileHDR() && !RHICmdList.IsInsideRenderPass())
@@ -193,6 +274,9 @@ void FGoogleARCorePassthroughCameraRenderer::RenderVideoOverlayWithMaterial(FRHI
 		TShaderRef<FGoogleARCoreCameraOverlayPS> PixelShader = MaterialShaderMap->GetShader<FGoogleARCoreCameraOverlayPS>();
 		TShaderRef<FGoogleARCoreCameraOverlayVS> VertexShader = MaterialShaderMap->GetShader<FGoogleARCoreCameraOverlayVS>();
 
+		const bool bCameraYCbCrConversionEnabled = bUsesCameraTexture && CameraYCbCrConversion.YCbCrModelConversion != EYCbCrModelConversion::None;
+		TShaderRef<FGoogleARCoreCameraOverlayYCbCrConversionPS> PixelShaderWithYCbCrConversion = MaterialShaderMap->GetShader<FGoogleARCoreCameraOverlayYCbCrConversionPS>();
+
 		FGraphicsPipelineStateInitializer GraphicsPSOInit;
 		RHICmdList.ApplyCachedRenderTargets(GraphicsPSOInit);
 
@@ -215,13 +299,28 @@ void FGoogleARCorePassthroughCameraRenderer::RenderVideoOverlayWithMaterial(FRHI
 
 		GraphicsPSOInit.BoundShaderState.VertexDeclarationRHI = GFilterVertexDeclaration.VertexDeclarationRHI;
 		GraphicsPSOInit.BoundShaderState.VertexShaderRHI = VertexShader.GetVertexShader();
-		GraphicsPSOInit.BoundShaderState.PixelShaderRHI = PixelShader.GetPixelShader();
+		if (bCameraYCbCrConversionEnabled)
+		{
+			GraphicsPSOInit.BoundShaderState.PixelShaderRHI = PixelShaderWithYCbCrConversion.GetPixelShader();
+		}
+		else
+		{
+			GraphicsPSOInit.BoundShaderState.PixelShaderRHI = PixelShader.GetPixelShader();
+		}
 		GraphicsPSOInit.PrimitiveType = PT_TriangleList;
 
 		SetGraphicsPipelineState(RHICmdList, GraphicsPSOInit, 0);
 
 		SetShaderParametersLegacyVS(RHICmdList, VertexShader, InView);
-		SetShaderParametersLegacyPS(RHICmdList, PixelShader, InView, MaterialProxy, CameraMaterial);
+		if (bCameraYCbCrConversionEnabled)
+		{
+			const FGoogleARCoreCameraOverlayYCbCrConversionPS::FYCbCrConversionParameters YCbCrConversionParameters = PixelShaderWithYCbCrConversion->GetYCbCrConversionParameters(CameraYCbCrConversion);
+			SetShaderParametersMixedPS(RHICmdList, PixelShaderWithYCbCrConversion, YCbCrConversionParameters, InView, MaterialProxy, CameraMaterial);
+		}
+		else
+		{
+			SetShaderParametersLegacyPS(RHICmdList, PixelShader, InView, MaterialProxy, CameraMaterial);
+		}
 
 		if (OverlayVertexBufferRHI.IsValid() && OverlayIndexBufferRHI.IsValid())
 		{
@@ -255,7 +354,10 @@ void FGoogleARCorePassthroughCameraRenderer::RenderVideoOverlay_RenderThread(FRH
 	}
 #endif
 	
-	RenderVideoOverlayWithMaterial(RHICmdList, InView, OverlayMaterialToUse, false);
+	// Debug Overlay materials use the depth image, so they don't need to use YCbCr conversion.
+	const bool bUsesCameraTexture = (OverlayMaterialToUse == RegularOverlayMaterial);
+
+	RenderVideoOverlayWithMaterial(RHICmdList, InView, OverlayMaterialToUse, false, bUsesCameraTexture);
 	
 #if ALLOWS_DEBUG_OVERLAY
 	if (GDebugOverlayMode)
@@ -268,7 +370,7 @@ void FGoogleARCorePassthroughCameraRenderer::RenderVideoOverlay_RenderThread(FRH
 	if (bEnableOcclusionRendering && DepthOcclusionMaterial)
 	{
 		UARUtilitiesFunctionLibrary::UpdateWorldToMeterScale(DepthOcclusionMaterial, 100.f);
-		RenderVideoOverlayWithMaterial(RHICmdList, InView, DepthOcclusionMaterial, true);
+		RenderVideoOverlayWithMaterial(RHICmdList, InView, DepthOcclusionMaterial, true, true);
 	}
 }
 
@@ -303,6 +405,18 @@ void FGoogleARCorePassthroughCameraRenderer::UpdateCameraTextures(UTexture* NewC
 	UARUtilitiesFunctionLibrary::UpdateCameraTextureParam(RegularOverlayMaterial, NewCameraTexture);
 }
 
+void FGoogleARCorePassthroughCameraRenderer::UpdateCameraYCbCrConversion(IYCbCrConversionQuery* NewYCbCrConversionQuery)
+{
+	if (NewYCbCrConversionQuery)
+	{
+		ENQUEUE_RENDER_COMMAND(SetCameraYCbCrConversion)(
+			[this, NewYCbCrConversionQuery](FRHICommandListImmediate& RHICmdList)
+			{
+				CameraYCbCrConversion = NewYCbCrConversionQuery->GetYCbCrConversion_RenderThread();
+			});
+	}
+}
+
 void FGoogleARCorePassthroughCameraRenderer::AddReferencedObjects(FReferenceCollector& Collector)
 {
 	Collector.AddReferencedObject(RegularOverlayMaterial);
diff --git a/Engine/Plugins/Runtime/AR/Google/GoogleARCore/Source/GoogleARCoreRendering/Public/GoogleARCorePassthroughCameraRenderer.h b/Engine/Plugins/Runtime/AR/Google/GoogleARCore/Source/GoogleARCoreRendering/Public/GoogleARCorePassthroughCameraRenderer.h
--- a/Engine/Plugins/Runtime/AR/Google/GoogleARCore/Source/GoogleARCoreRendering/Public/GoogleARCorePassthroughCameraRenderer.h
+++ b/Engine/Plugins/Runtime/AR/Google/GoogleARCore/Source/GoogleARCoreRendering/Public/GoogleARCorePassthroughCameraRenderer.h
@@ -2,6 +2,7 @@
 
 #pragma once
 
+#include "GoogleARCoreYCbCrConversion.h"
 #include "CoreMinimal.h"
 #include "UObject/ConstructorHelpers.h"
 #include "RHI.h"
@@ -71,6 +72,8 @@ public:
 	void RenderVideoOverlay_RenderThread(FRHICommandList& RHICmdList, FSceneView& InView);
 	
 	void UpdateCameraTextures(UTexture* NewCameraTexture, UTexture* DepthTexture, bool bEnableOcclusion);
+
+	void UpdateCameraYCbCrConversion(IYCbCrConversionQuery* NewYCbCrConversionQuery);
 	
 	void AddReferencedObjects(FReferenceCollector& Collector) override;
 	virtual FString GetReferencerName() const override
@@ -79,7 +82,7 @@ public:
 	}
 
 private:
-	void RenderVideoOverlayWithMaterial(FRHICommandList& RHICmdList, FSceneView& InView, UMaterialInstanceDynamic* OverlayMaterialToUse, bool bRenderingOcclusion);
+	void RenderVideoOverlayWithMaterial(FRHICommandList& RHICmdList, FSceneView& InView, UMaterialInstanceDynamic* OverlayMaterialToUse, bool bRenderingOcclusion, bool bUsesCameraTexture);
 
 private:
 	FBufferRHIRef OverlayIndexBufferRHI;
@@ -91,4 +94,6 @@ private:
 	TObjectPtr<UMaterialInstanceDynamic> DepthOcclusionMaterial = nullptr;
 	
 	bool bEnableOcclusionRendering = false;
+
+	FYCbCrConversion CameraYCbCrConversion;
 };
diff --git a/Engine/Plugins/Runtime/AR/Google/GoogleARCore/Source/GoogleARCoreRendering/Public/GoogleARCoreYCbCrConversion.h b/Engine/Plugins/Runtime/AR/Google/GoogleARCore/Source/GoogleARCoreRendering/Public/GoogleARCoreYCbCrConversion.h
new file mode 100644
--- /dev/null
+++ b/Engine/Plugins/Runtime/AR/Google/GoogleARCore/Source/GoogleARCoreRendering/Public/GoogleARCoreYCbCrConversion.h
@@ -0,0 +1,41 @@
+// Copyright Epic Games, Inc. All Rights Reserved.
+
+#pragma once
+
+#include "CoreMinimal.h"
+
+enum class EYCbCrModelConversion : uint8
+{
+	None = 0,
+	YCbCrIdentity,
+	YCbCrRec709,
+	YCbCrRec601,
+	YCbCrRec2020,
+	Max
+};
+
+enum class EYCbCrRange : uint8
+{
+	Unknown = 0,
+	Full,
+	Narrow,
+	Max
+};
+
+struct FYCbCrConversion
+{
+
+	EYCbCrModelConversion YCbCrModelConversion = EYCbCrModelConversion::None;
+	EYCbCrRange YCbCrRange = EYCbCrRange::Unknown;
+	uint8 NumBits = 8;
+};
+
+class IYCbCrConversionQuery
+{
+public:
+	virtual ~IYCbCrConversionQuery() = default;
+	
+	// Returns the YCbCr conversion paramers of Vulkan's external Hardware Buffer.
+	// It should be called from the render thread only.
+	virtual FYCbCrConversion GetYCbCrConversion_RenderThread() const = 0;
+};
diff --git a/Engine/Shaders/Private/PostProcessMaterialShaders.usf b/Engine/Shaders/Private/PostProcessMaterialShaders.usf
--- a/Engine/Shaders/Private/PostProcessMaterialShaders.usf
+++ b/Engine/Shaders/Private/PostProcessMaterialShaders.usf
@@ -8,6 +8,17 @@
 #include "ScreenPass.ush"
 #include "VirtualTextureCommon.ush"
 
+#ifndef POST_PROCESS_AR_YCBCR_CONVERSION
+#define POST_PROCESS_AR_YCBCR_CONVERSION 0
+#endif
+
+#if POST_PROCESS_AR_YCBCR_CONVERSION
+#include "ColorUtils.ush"
+
+float4x4 YCbCrColorTransform;
+uint YCbCrSrgbToLinear;
+#endif
+
 #pragma message("UESHADERMETADATA_VERSION 02E0E449-2223-45B4-8CD0-93AB087FBEE3")
 
 #ifndef POST_PROCESS_MATERIAL
@@ -531,6 +542,11 @@ void MainPS(
 	// Grab emissive colour as output
 	EmissiveColor = GetMaterialEmissive(PixelMaterialInputs);
 	
+#if POST_PROCESS_AR_YCBCR_CONVERSION
+	const float3 YCbCr = EmissiveColor.yzx;
+	EmissiveColor.rgb = YuvToRgb(YCbCr, YCbCrColorTransform, YCbCrSrgbToLinear);
+#endif
+	
   #if MATERIAL_OUTPUT_OPACITY_AS_ALPHA
 	Alpha = GetMaterialOpacity(PixelMaterialInputs);
   #else
