Obsoletion pass for DeterministicWarningLevel & DeprecationWarningLevel

diff --git a/Engine/Plugins/Runtime/AR/Google/GoogleARCore/Source/GoogleARCoreBase/GoogleARCoreBase.Build.cs b/Engine/Plugins/Runtime/AR/Google/GoogleARCore/Source/GoogleARCoreBase/GoogleARCoreBase.Build.cs
--- a/Engine/Plugins/Runtime/AR/Google/GoogleARCore/Source/GoogleARCoreBase/GoogleARCoreBase.Build.cs
+++ b/Engine/Plugins/Runtime/AR/Google/GoogleARCore/Source/GoogleARCoreBase/GoogleARCoreBase.Build.cs
@@ -10,7 +10,7 @@ namespace UnrealBuildTool.Rules
 	{
 		public GoogleARCoreBase(ReadOnlyTargetRules Target) : base(Target)
 		{
-			DeterministicWarningLevel = WarningLevel.Off; // __DATE__ __TIME__ in Private/GoogleARCoreXRTrackingSystem.cpp
+			CppCompileWarningSettings.DeterministicWarningLevel = WarningLevel.Off; // __DATE__ __TIME__ in Private/GoogleARCoreXRTrackingSystem.cpp
 
 			PublicIncludePathModuleNames.Add("TargetPlatform");
 
diff --git a/Engine/Source/Programs/UnrealBuildTool/Configuration/ModuleRules.cs b/Engine/Source/Programs/UnrealBuildTool/Configuration/ModuleRules.cs
--- a/Engine/Source/Programs/UnrealBuildTool/Configuration/ModuleRules.cs
+++ b/Engine/Source/Programs/UnrealBuildTool/Configuration/ModuleRules.cs
@@ -969,6 +969,7 @@ namespace UnrealBuildTool
 		/// <summary>
 		/// How to treat deterministic warnings (experimental).
 		/// </summary>
+		[Obsolete("Deprecated in UE5.6 - Replace with ModuleRules.CppCompileWarningSettings.DeterministicWarningLevel")]
 		public WarningLevel DeterministicWarningLevel
 		{
 			get => CppCompileWarningSettings.DeterministicWarningLevel;
diff --git a/Engine/Source/Programs/UnrealBuildTool/Configuration/ReadOnlyTargetRules.cs b/Engine/Source/Programs/UnrealBuildTool/Configuration/ReadOnlyTargetRules.cs
--- a/Engine/Source/Programs/UnrealBuildTool/Configuration/ReadOnlyTargetRules.cs
+++ b/Engine/Source/Programs/UnrealBuildTool/Configuration/ReadOnlyTargetRules.cs
@@ -361,7 +361,8 @@ namespace UnrealBuildTool
 
 		public WarningLevel DefaultWarningLevel => Inner.DefaultWarningLevel;
 
-		public WarningLevel DeprecationWarningLevel => Inner.DeprecationWarningLevel;
+		[Obsolete("Deprecated in UE5.6 - Replace with ReadOnlyTargetRules.CppCompileWarningSettings.DeprecationWarningLevel")]
+		public WarningLevel DeprecationWarningLevel => Inner.CppCompileWarningSettings.DeprecationWarningLevel;
 
 		public ReadOnlyCppCompileWarnings CppCompileWarningSettings { get; init; }
 
diff --git a/Engine/Source/Programs/UnrealBuildTool/Configuration/TargetRules.cs b/Engine/Source/Programs/UnrealBuildTool/Configuration/TargetRules.cs
--- a/Engine/Source/Programs/UnrealBuildTool/Configuration/TargetRules.cs
+++ b/Engine/Source/Programs/UnrealBuildTool/Configuration/TargetRules.cs
@@ -1839,6 +1839,7 @@ namespace UnrealBuildTool
 		/// <summary>
 		/// Level to report deprecation warnings as errors
 		/// </summary>
+		[Obsolete("Deprecated in UE5.6 - Replace with TargetRules.CppCompileWarningSettings.DeprecationWarningLevel")]
 		public WarningLevel DeprecationWarningLevel
 		{
 			get => CppCompileWarningSettings.DeprecationWarningLevel;
diff --git a/Engine/Source/Runtime/BuildSettings/BuildSettings.Build.cs b/Engine/Source/Runtime/BuildSettings/BuildSettings.Build.cs
--- a/Engine/Source/Runtime/BuildSettings/BuildSettings.Build.cs
+++ b/Engine/Source/Runtime/BuildSettings/BuildSettings.Build.cs
@@ -10,7 +10,7 @@ public class BuildSettings : ModuleRules
 {
 	public BuildSettings(ReadOnlyTargetRules Target) : base(Target)
 	{
-		DeterministicWarningLevel = WarningLevel.Off; // This module intentionally uses __DATE__ and __TIME__ macros
+		CppCompileWarningSettings.DeterministicWarningLevel = WarningLevel.Off; // This module intentionally uses __DATE__ and __TIME__ macros
 		CppCompileWarningSettings.UnsafeTypeCastWarningLevel = WarningLevel.Error;
 
 		PrivateIncludePathModuleNames.Add("Core");
diff --git a/Engine/Source/Runtime/Datasmith/CADKernel/Base/CADKernel.Build.cs b/Engine/Source/Runtime/Datasmith/CADKernel/Base/CADKernel.Build.cs
--- a/Engine/Source/Runtime/Datasmith/CADKernel/Base/CADKernel.Build.cs
+++ b/Engine/Source/Runtime/Datasmith/CADKernel/Base/CADKernel.Build.cs
@@ -9,7 +9,7 @@ namespace UnrealBuildTool.Rules
 			: base(Target)
 		{
 			//OptimizeCode = CodeOptimization.Never;
-			DeterministicWarningLevel = WarningLevel.Off; // __DATE__ in Private/CADKernel/Core/System.cpp
+			CppCompileWarningSettings.DeterministicWarningLevel = WarningLevel.Off; // __DATE__ in Private/CADKernel/Core/System.cpp
 			
 			PublicDefinitions.Add("CADKERNEL_THINZONE=0");
 
